[mutable] Renamed the plugin "DataflowMutable" to "MutableDataflow"

#rnx

diff --git a/Engine/Plugins/Experimental/DataflowMutable/DataflowMutable.uplugin b/Engine/Plugins/Experimental/DataflowMutable/DataflowMutable.uplugin
deleted file mode 100644
diff --git a/Engine/Plugins/Experimental/DataflowMutable/Source/DataflowMutableEditor/DataflowMutableEditor.Build.cs b/Engine/Plugins/Experimental/DataflowMutable/Source/DataflowMutableEditor/DataflowMutableEditor.Build.cs
deleted file mode 100644
diff --git a/Engine/Plugins/Experimental/DataflowMutable/Source/DataflowMutableEditor/Private/DataflowMutableEditorModule.cpp b/Engine/Plugins/Experimental/DataflowMutable/Source/DataflowMutableEditor/Private/DataflowMutableEditorModule.cpp
deleted file mode 100644
diff --git a/Engine/Plugins/Experimental/DataflowMutable/Source/DataflowMutableEditor/Private/Nodes/COInstanceGeneratorNode.cpp b/Engine/Plugins/Experimental/DataflowMutable/Source/DataflowMutableEditor/Private/Nodes/COInstanceGeneratorNode.cpp
deleted file mode 100644
diff --git a/Engine/Plugins/Experimental/DataflowMutable/Source/DataflowMutableEditor/Private/Nodes/MakeMutableMaterialParametersArrayNode.cpp b/Engine/Plugins/Experimental/DataflowMutable/Source/DataflowMutableEditor/Private/Nodes/MakeMutableMaterialParametersArrayNode.cpp
deleted file mode 100644
diff --git a/Engine/Plugins/Experimental/DataflowMutable/Source/DataflowMutableEditor/Private/Nodes/MakeMutableSkeletalMeshParametersArrayNode.cpp b/Engine/Plugins/Experimental/DataflowMutable/Source/DataflowMutableEditor/Private/Nodes/MakeMutableSkeletalMeshParametersArrayNode.cpp
deleted file mode 100644
diff --git a/Engine/Plugins/Experimental/DataflowMutable/Source/DataflowMutableEditor/Private/Nodes/MakeMutableTextureParametersArrayNode.cpp b/Engine/Plugins/Experimental/DataflowMutable/Source/DataflowMutableEditor/Private/Nodes/MakeMutableTextureParametersArrayNode.cpp
deleted file mode 100644
diff --git a/Engine/Plugins/Experimental/DataflowMutable/Source/DataflowMutableEditor/Private/Nodes/MutableMaterialParameterNode.cpp b/Engine/Plugins/Experimental/DataflowMutable/Source/DataflowMutableEditor/Private/Nodes/MutableMaterialParameterNode.cpp
deleted file mode 100644
diff --git a/Engine/Plugins/Experimental/DataflowMutable/Source/DataflowMutableEditor/Private/Nodes/MutableSkeletalMeshParameterNode.cpp b/Engine/Plugins/Experimental/DataflowMutable/Source/DataflowMutableEditor/Private/Nodes/MutableSkeletalMeshParameterNode.cpp
deleted file mode 100644
diff --git a/Engine/Plugins/Experimental/DataflowMutable/Source/DataflowMutableEditor/Private/Nodes/MutableTextureParameterNode.cpp b/Engine/Plugins/Experimental/DataflowMutable/Source/DataflowMutableEditor/Private/Nodes/MutableTextureParameterNode.cpp
deleted file mode 100644
diff --git a/Engine/Plugins/Experimental/DataflowMutable/Source/DataflowMutableEditor/Public/DataflowMutableEditorModule.h b/Engine/Plugins/Experimental/DataflowMutable/Source/DataflowMutableEditor/Public/DataflowMutableEditorModule.h
deleted file mode 100644
diff --git a/Engine/Plugins/Experimental/DataflowMutable/Source/DataflowMutableEditor/Public/MutableDataflowParameters.h b/Engine/Plugins/Experimental/DataflowMutable/Source/DataflowMutableEditor/Public/MutableDataflowParameters.h
deleted file mode 100644
diff --git a/Engine/Plugins/Experimental/DataflowMutable/Source/DataflowMutableEditor/Public/Nodes/COInstanceGeneratorNode.h b/Engine/Plugins/Experimental/DataflowMutable/Source/DataflowMutableEditor/Public/Nodes/COInstanceGeneratorNode.h
deleted file mode 100644
diff --git a/Engine/Plugins/Experimental/DataflowMutable/Source/DataflowMutableEditor/Public/Nodes/MakeMutableMaterialParametersArrayNode.h b/Engine/Plugins/Experimental/DataflowMutable/Source/DataflowMutableEditor/Public/Nodes/MakeMutableMaterialParametersArrayNode.h
deleted file mode 100644
diff --git a/Engine/Plugins/Experimental/DataflowMutable/Source/DataflowMutableEditor/Public/Nodes/MakeMutableParametersArrayBaseNode.h b/Engine/Plugins/Experimental/DataflowMutable/Source/DataflowMutableEditor/Public/Nodes/MakeMutableParametersArrayBaseNode.h
deleted file mode 100644
diff --git a/Engine/Plugins/Experimental/DataflowMutable/Source/DataflowMutableEditor/Public/Nodes/MakeMutableSkeletalMeshParametersArrayNode.h b/Engine/Plugins/Experimental/DataflowMutable/Source/DataflowMutableEditor/Public/Nodes/MakeMutableSkeletalMeshParametersArrayNode.h
deleted file mode 100644
diff --git a/Engine/Plugins/Experimental/DataflowMutable/Source/DataflowMutableEditor/Public/Nodes/MakeMutableTextureParametersArrayNode.h b/Engine/Plugins/Experimental/DataflowMutable/Source/DataflowMutableEditor/Public/Nodes/MakeMutableTextureParametersArrayNode.h
deleted file mode 100644
diff --git a/Engine/Plugins/Experimental/DataflowMutable/Source/DataflowMutableEditor/Public/Nodes/MutableMaterialParameterNode.h b/Engine/Plugins/Experimental/DataflowMutable/Source/DataflowMutableEditor/Public/Nodes/MutableMaterialParameterNode.h
deleted file mode 100644
diff --git a/Engine/Plugins/Experimental/DataflowMutable/Source/DataflowMutableEditor/Public/Nodes/MutableSkeletalMeshParameterNode.h b/Engine/Plugins/Experimental/DataflowMutable/Source/DataflowMutableEditor/Public/Nodes/MutableSkeletalMeshParameterNode.h
deleted file mode 100644
diff --git a/Engine/Plugins/Experimental/DataflowMutable/Source/DataflowMutableEditor/Public/Nodes/MutableTextureParameterNode.h b/Engine/Plugins/Experimental/DataflowMutable/Source/DataflowMutableEditor/Public/Nodes/MutableTextureParameterNode.h
deleted file mode 100644
diff --git a/Engine/Plugins/Experimental/MutableDataflow/MutableDataflow.uplugin b/Engine/Plugins/Experimental/MutableDataflow/MutableDataflow.uplugin
new file mode 100644
--- /dev/null
+++ b/Engine/Plugins/Experimental/MutableDataflow/MutableDataflow.uplugin
@@ -0,0 +1,38 @@
+{
+	"FileVersion": 3,
+	"Version": 1,
+	"VersionName": "0.1",
+	"FriendlyName": "Dataflow Extensions For Mutable",
+	"Description": "Adds Mutable functionality to work with Dataflow objects from the Dataflow plugin",
+	"Category": "Animation",
+	"CreatedBy": "Epic Games, Inc.",
+	"CreatedByURL": "https://epicgames.com",
+	"DocsURL": "",
+	"MarketplaceURL": "",
+	"SupportURL": "",
+	"CanContainContent": false,
+	"IsBetaVersion": false,
+	"IsExperimentalVersion": true,
+	"Installed": false,
+	"EnabledByDefault": true,
+	"Modules": 
+	[
+		{
+			"Name": "MutableDataflowEditor",
+			"Type": "UncookedOnly",
+			"LoadingPhase": "Default"
+		}
+	],
+	"Plugins": 
+	[
+		{
+			"Name": "Dataflow",
+			"Enabled": true
+		},
+		{
+			"Name": "Mutable",
+			"Enabled": true
+		}
+	]
+
+}
\ No newline at end of file
diff --git a/Engine/Plugins/Experimental/MutableDataflow/Source/MutableDataflowEditor/MutableDataflowEditor.Build.cs b/Engine/Plugins/Experimental/MutableDataflow/Source/MutableDataflowEditor/MutableDataflowEditor.Build.cs
new file mode 100644
--- /dev/null
+++ b/Engine/Plugins/Experimental/MutableDataflow/Source/MutableDataflowEditor/MutableDataflowEditor.Build.cs
@@ -0,0 +1,31 @@
+// Copyright Epic Games, Inc. All Rights Reserved.
+
+namespace UnrealBuildTool.Rules
+{
+	public class MutableDataflowEditor : ModuleRules
+	{
+        public MutableDataflowEditor(ReadOnlyTargetRules Target) : base(Target)
+		{
+			PrivateIncludePaths.Add(ModuleDirectory + "/Private");
+			PublicIncludePaths.Add(ModuleDirectory + "/Public");
+			
+			PublicDependencyModuleNames.AddRange(
+				new string[]
+				{
+					"Core",
+					"CoreUObject",
+					"Engine",
+					"CustomizableObject",
+				});
+
+			PrivateDependencyModuleNames.AddRange(
+				new string[]
+				{
+					"DataflowCore",
+					"DataflowEditor",
+					"DataflowEngine",
+					"DataflowNodes",
+				});
+		}
+	}
+}
diff --git a/Engine/Plugins/Experimental/MutableDataflow/Source/MutableDataflowEditor/Private/MutableDataflowEditorModule.cpp b/Engine/Plugins/Experimental/MutableDataflow/Source/MutableDataflowEditor/Private/MutableDataflowEditorModule.cpp
new file mode 100644
--- /dev/null
+++ b/Engine/Plugins/Experimental/MutableDataflow/Source/MutableDataflowEditor/Private/MutableDataflowEditorModule.cpp
@@ -0,0 +1,41 @@
+// Copyright Epic Games, Inc. All Rights Reserved.
+
+#include "MutableDataflowEditorModule.h"
+
+#include "Dataflow/DataflowAnyTypeRegistry.h"
+#include "Dataflow/DataflowNode.h"
+#include "Dataflow/DataflowCore.h"
+#include "Nodes/COInstanceGeneratorNode.h"
+#include "Nodes/MutableMaterialParameterNode.h"
+#include "Nodes/MutableSkeletalMeshParameterNode.h"
+#include "Nodes/MutableTextureParameterNode.h"
+#include "Nodes/MakeMutableMaterialParametersArrayNode.h"
+#include "Nodes/MakeMutableSkeletalMeshParametersArrayNode.h"
+#include "Nodes/MakeMutableTextureParametersArrayNode.h"
+
+IMPLEMENT_MODULE(FMutableDataflowEditorModule, MutableDataflowEditor);
+
+void FMutableDataflowEditorModule::StartupModule()
+{
+	// Main nodes
+	DATAFLOW_NODE_REGISTER_CREATION_FACTORY(FCOInstanceGeneratorNode);
+	DATAFLOW_NODE_REGISTER_CREATION_FACTORY(FCOInstanceGetComponentMesh);
+
+	// Auto connection between an array of FMutableGeneratedResource to a single TObjectPtr<USkeletalMesh>
+	UE_DATAFLOW_REGISTER_AUTOCONVERT(TArray<FMutableGeneratedResource>, TObjectPtr<USkeletalMesh>, FCOInstanceGetComponentMesh);
+	
+	// Parameter nodes
+	DATAFLOW_NODE_REGISTER_CREATION_FACTORY(FMutableSkeletalMeshParameterNode);
+	DATAFLOW_NODE_REGISTER_CREATION_FACTORY(FMutableTextureParameterNode);
+	DATAFLOW_NODE_REGISTER_CREATION_FACTORY(FMutableMaterialParameterNode);
+
+	// Parameter array generation nodes
+	DATAFLOW_NODE_REGISTER_CREATION_FACTORY(FMakeMutableSkeletalMeshParametersArrayNode);
+	DATAFLOW_NODE_REGISTER_CREATION_FACTORY(FMakeMutableMaterialParametersArrayNode);
+	DATAFLOW_NODE_REGISTER_CREATION_FACTORY(FMakeMutableTextureParametersArrayNode);
+
+	// Parameter node to array automatic conversion nodes
+	UE_DATAFLOW_REGISTER_AUTOCONVERT(FMutableSkeletalMeshParameter, TArray<FMutableSkeletalMeshParameter>, FMakeMutableSkeletalMeshParametersArrayNode);
+	UE_DATAFLOW_REGISTER_AUTOCONVERT(FMutableMaterialParameter, TArray<FMutableMaterialParameter>, FMakeMutableMaterialParametersArrayNode);
+	UE_DATAFLOW_REGISTER_AUTOCONVERT(FMutableTextureParameter, TArray<FMutableTextureParameter>, FMakeMutableTextureParametersArrayNode);
+}
diff --git a/Engine/Plugins/Experimental/MutableDataflow/Source/MutableDataflowEditor/Private/Nodes/COInstanceGeneratorNode.cpp b/Engine/Plugins/Experimental/MutableDataflow/Source/MutableDataflowEditor/Private/Nodes/COInstanceGeneratorNode.cpp
new file mode 100644
--- /dev/null
+++ b/Engine/Plugins/Experimental/MutableDataflow/Source/MutableDataflowEditor/Private/Nodes/COInstanceGeneratorNode.cpp
@@ -0,0 +1,320 @@
+// Copyright Epic Games, Inc. All Rights Reserved.
+
+#include "../Public/Nodes/COInstanceGeneratorNode.h"
+
+#include "MuCO/CustomizableObject.h"
+#include "MuCO/CustomizableObjectInstance.h"
+#include "Engine/SkeletalMesh.h"
+#include "MutableDataflowParameters.h"
+
+
+FCOInstanceGetComponentMesh::FCOInstanceGetComponentMesh(const UE::Dataflow::FNodeParameters& InParam, FGuid InGuid)
+	: Super(InParam, InGuid)
+{
+	RegisterInputConnection(&GeneratedResources);
+	RegisterInputConnection(&ComponentName);
+
+	RegisterOutputConnection(&ComponentSkeletalMesh);
+}
+
+
+void FCOInstanceGetComponentMesh::Evaluate(UE::Dataflow::FContext& Context, const FDataflowOutput* Out) const
+{
+	const TArray<FMutableGeneratedResource> InputGeneratedResources = GetValue(Context, &GeneratedResources);
+	const FString TargetComponentName = GetValue(Context, &ComponentName);
+
+	TObjectPtr<USkeletalMesh> OutputSkeletalMesh = nullptr;
+	if (!TargetComponentName.IsEmpty())
+	{
+		for (const FMutableGeneratedResource& Resource : InputGeneratedResources)
+		{
+			if (Resource.ComponentName.Equals(TargetComponentName))
+			{
+				OutputSkeletalMesh = Resource.SkeletalMesh;
+				break;
+			}
+		}
+	}
+	
+	SetValue(Context, OutputSkeletalMesh, &ComponentSkeletalMesh);
+}
+
+
+
+FCOInstanceGeneratorNode::FCOInstanceGeneratorNode(const UE::Dataflow::FNodeParameters& InParam,FGuid InGuid)
+: FDataflowNode(InParam, InGuid),
+	GenerateInstanceResources( FDataflowFunctionProperty::FDelegate::CreateRaw(this, &FCOInstanceGeneratorNode::OnGenerateInstanceResourcesRequested))
+{
+	// Inputs
+	RegisterInputConnection(&CustomizableObject);
+	RegisterInputConnection(&SkeletalMeshParameters);
+	RegisterInputConnection(&TextureParameters);
+	RegisterInputConnection(&MaterialParameters);
+	
+	// Outputs
+	RegisterOutputConnection(&GeneratedResources);
+	RegisterOutputConnection(&GeneratedSkeletalMeshes);
+}
+
+
+void FCOInstanceGeneratorNode::OnPropertyChanged(UE::Dataflow::FContext& Context, const FPropertyChangedEvent& InPropertyChangedEvent)
+{
+	FDataflowNode::OnPropertyChanged(Context, InPropertyChangedEvent);
+
+	if (const FProperty* ModifiedProperty = InPropertyChangedEvent.Property)
+	{
+		if (ModifiedProperty->GetFName() == GET_MEMBER_NAME_CHECKED(FCOInstanceGeneratorNode, CustomizableObject))
+		{
+			ClearCachedParameters();
+			
+			CustomizableObjectInstance = nullptr;
+			
+			GeneratedResources.Reset();
+			GeneratedSkeletalMeshes.Reset();
+			Invalidate();
+		}
+	}
+}
+
+
+void FCOInstanceGeneratorNode::Evaluate(UE::Dataflow::FContext& Context, const FDataflowOutput* Out) const
+{
+	if (!GetValue(Context, &CustomizableObject))
+	{
+		Context.Error(TEXT("Unable to evaluate : No Customizable Object has been set"), this);
+		return;
+	}
+	
+	if (!CustomizableObjectInstance)
+	{
+		Context.Error(TEXT("Unable to evaluate : Press \"Generate Resources\" before evaluating the node"), this);
+		return;
+	}
+
+	if (bIsGeneratingResources)
+	{
+		Context.Warning(TEXT("Unable to evaluate : Resource generation is in process. Please hold"), this);
+		return;
+	}
+	
+	if (Out->IsA<TArray<FMutableGeneratedResource>>(&GeneratedResources))
+	{
+		SetValue(Context, GeneratedResources, &GeneratedResources);
+	}
+	else if (Out->IsA<TArray<TObjectPtr<USkeletalMesh>>>(&GeneratedSkeletalMeshes))
+	{
+		// Extract the generated meshes from the generated resources collection as there we have them already set
+		TArray<TObjectPtr<USkeletalMesh>> SkeletalMeshes;
+		SkeletalMeshes.Reserve(GeneratedResources.Num());
+		for (const FMutableGeneratedResource& Resource : GeneratedResources)
+		{
+			SkeletalMeshes.Add(Resource.SkeletalMesh);
+		}
+		SetValue(Context, SkeletalMeshes, &GeneratedSkeletalMeshes);
+	}
+}
+
+
+void FCOInstanceGeneratorNode::ClearCachedParameters()
+{
+	CachedSkeletalMeshParameters.Reset();
+	CachedMaterialParameters.Reset();
+	CachedTextureParameters.Reset();
+}
+
+
+void FCOInstanceGeneratorNode::OnGenerateInstanceResourcesRequested(UE::Dataflow::FContext& Context)
+{
+	if (bIsGeneratingResources)
+	{
+		Context.Warning(TEXT("Unable to generate resources. Resource generation is already in process"), this);
+		return;
+	}
+	
+	if (GetValue(Context, &CustomizableObject))
+	{
+		// Prevent the scheduling or more compilations while the generation of resources is in process
+		bIsGeneratingResources = true;
+		RequestCompilation(Context);
+	}
+	else
+	{
+		Context.Error(TEXT("Unable to generate resources. The Customizable object has not been set"), this);
+	}
+}
+
+
+void FCOInstanceGeneratorNode::RequestCompilation(UE::Dataflow::FContext& Context)
+{
+	// Capture the inputs of the node (target parameters) so we can apply them after the compilation without us having to rely on the context (as it may be gone)
+	CacheNodeInputs(Context);
+	check(CustomizableObject);
+	
+	FCompileNativeDelegate CompilationDelegate;
+	TSharedRef<FDataflowNode> This = AsShared();
+	CompilationDelegate.BindLambda([This](const FCompileCallbackParams& CallbackParams)->void
+	{
+		FCOInstanceGeneratorNode* HostNode = static_cast<FCOInstanceGeneratorNode*>(&This.Get());
+		check(HostNode);
+			
+		// First generate the instance
+		HostNode->CustomizableObjectInstance = HostNode->CustomizableObject->CreateInstance();
+		check(HostNode->CustomizableObjectInstance)
+			
+		// Apply the parameters set in the node into the instance
+		HostNode->ApplyInstanceParameters(*HostNode->CustomizableObjectInstance);
+
+		// once the compilation has been performed (if required) ask for the update of the COI
+		HostNode->RequestUpdate(*HostNode->CustomizableObjectInstance);
+	});
+
+	// Request the compilation of the CO
+	FCompileParams CompilationParams;
+	CompilationParams.bAsync = true;
+	CompilationParams.CallbackNative = CompilationDelegate;
+	
+	// todo: Expose compilation errors and warnings to the Dataflow context
+	CustomizableObject->Compile(CompilationParams);
+
+	// Prevent the scheduling or more compilations while the generation of resources is in process
+	bIsGeneratingResources = true;
+}
+
+
+void FCOInstanceGeneratorNode::CacheNodeInputs(UE::Dataflow::FContext& Context)
+{
+	ClearCachedParameters();
+
+	// Customizable Object
+	CustomizableObject = GetValue(Context, &CustomizableObject);
+
+	// Cache the parameters now that we have the context. Do not cache duplicate entries (based on the param name)
+	
+	// Texture Parameters
+	{
+		const TArray<FMutableTextureParameter>& TempTextureParametersArray = GetValue(Context, &TextureParameters);
+		for (const FMutableTextureParameter& TextureParameter : TempTextureParametersArray)
+		{
+			if (CachedTextureParameters.Contains(TextureParameter))
+			{
+				Context.Warning(FString::Printf( TEXT("Multiple Texture Parameters have been found sharing the \"%s\" name , only the first parameter will be used."), *TextureParameter.Name), this);
+				continue;
+			}
+		
+			CachedTextureParameters.Add(TextureParameter);
+		}
+	}
+	
+	// Mesh Parameters
+	{
+		const TArray<FMutableSkeletalMeshParameter>& TempMeshParametersArray = GetValue(Context, &SkeletalMeshParameters);
+		for (const FMutableSkeletalMeshParameter& MeshParameter : TempMeshParametersArray)
+		{
+			if (CachedSkeletalMeshParameters.Contains(MeshParameter))
+			{
+				Context.Warning(FString::Printf( TEXT("Multiple Skeletal Mesh Parameters have been found sharing the \"%s\" name , only the first parameter will be used."), *MeshParameter.Name), this);
+				continue;
+			}
+		
+			CachedSkeletalMeshParameters.Add(MeshParameter);
+		}
+	}
+	
+	// Material parameters
+	{
+		const TArray<FMutableMaterialParameter>& TempMaterialParametersArray = GetValue(Context, &MaterialParameters);
+		for (const FMutableMaterialParameter& MaterialParameter : TempMaterialParametersArray)
+		{
+			if (CachedMaterialParameters.Contains(MaterialParameter))
+			{
+				Context.Warning(FString::Printf( TEXT("Multiple Material Parameters have been found sharing the \"%s\" name , only the first parameter will be used."), *MaterialParameter.Name), this);
+				continue;
+			}
+		
+			CachedMaterialParameters.Add(MaterialParameter);
+		}
+	}
+}
+
+
+void FCOInstanceGeneratorNode::ApplyInstanceParameters(UCustomizableObjectInstance& Instance)
+{
+	// todo. notify when a parameter does not exist :
+	// In order to do that we will require to have async node evaluations and then access to the
+	// context to log the warning
+	
+	// Texture Parameters
+	for (const FMutableTextureParameter& TextureParameter : CachedTextureParameters)
+	{
+		const FString ParameterName = TextureParameter.Name;
+		if (Instance.ContainsTextureParameter(ParameterName))
+		{
+			Instance.SetTextureParameterSelectedOption(ParameterName, TextureParameter.Texture);
+		}
+	}
+	
+	// Mesh Parameters
+	for (const FMutableSkeletalMeshParameter& MeshParameter : CachedSkeletalMeshParameters)
+	{
+		const FString ParameterName = MeshParameter.Name;
+		if (Instance.ContainsSkeletalMeshParameter(ParameterName))
+		{
+			Instance.SetSkeletalMeshParameterSelectedOption(ParameterName, MeshParameter.Mesh);
+		}
+	}
+	
+	// Material parameters
+	for (const FMutableMaterialParameter& MaterialParameter : CachedMaterialParameters)
+	{
+		const FString ParameterName = MaterialParameter.Name;
+		if (Instance.ContainsMaterialParameter(ParameterName))
+		{
+			Instance.SetMaterialParameterSelectedOption(ParameterName, MaterialParameter.Material);				
+		}
+	}
+}
+
+
+void FCOInstanceGeneratorNode::RequestUpdate(UCustomizableObjectInstance& Instance)
+{
+	TSharedRef<FDataflowNode> This = AsShared();
+	
+	// Instance update delegate	
+	FInstanceUpdateNativeDelegate InstanceUpdateNativeDelegate;
+	InstanceUpdateNativeDelegate.AddLambda([This](FUpdateContext OutUpdateContext) -> void
+	{
+		FCOInstanceGeneratorNode* HostNode = static_cast<FCOInstanceGeneratorNode*>(&This.Get());
+		check(HostNode);
+			
+		HostNode->GeneratedResources = HostNode->GetInstanceGeneratedMeshes(*OutUpdateContext.Instance);
+			
+		// Tell the system it is ok to generate resources again
+		HostNode->bIsGeneratingResources = false;
+			
+		This->Invalidate();
+	});
+	 		
+	Instance.UpdateSkeletalMeshAsyncResult(InstanceUpdateNativeDelegate,true, true);
+}
+
+
+TArray<FMutableGeneratedResource> FCOInstanceGeneratorNode::GetInstanceGeneratedMeshes(const UCustomizableObjectInstance& Instance) const
+{
+	TArray<FMutableGeneratedResource> GeneratedMeshesData;
+			
+	const TArray<FName> InstanceComponentNames = Instance.GetComponentNames();
+	for (const FName& ComponentName : InstanceComponentNames)
+	{
+		if (const TObjectPtr<USkeletalMesh> SkeletalMesh = Instance.GetSkeletalMeshComponentSkeletalMesh(ComponentName))
+		{
+			FMutableGeneratedResource ComponentData;
+			ComponentData.ComponentName = ComponentName.ToString();
+			ComponentData.SkeletalMesh = SkeletalMesh;
+			
+			GeneratedMeshesData.Add(ComponentData);
+		}
+	}
+
+	return GeneratedMeshesData;
+}
+
diff --git a/Engine/Plugins/Experimental/MutableDataflow/Source/MutableDataflowEditor/Private/Nodes/MakeMutableMaterialParametersArrayNode.cpp b/Engine/Plugins/Experimental/MutableDataflow/Source/MutableDataflowEditor/Private/Nodes/MakeMutableMaterialParametersArrayNode.cpp
new file mode 100644
--- /dev/null
+++ b/Engine/Plugins/Experimental/MutableDataflow/Source/MutableDataflowEditor/Private/Nodes/MakeMutableMaterialParametersArrayNode.cpp
@@ -0,0 +1,50 @@
+﻿// Copyright Epic Games, Inc. All Rights Reserved.
+
+
+#include "Nodes/MakeMutableMaterialParametersArrayNode.h"
+
+
+FMakeMutableMaterialParametersArrayNode::FMakeMutableMaterialParametersArrayNode(const UE::Dataflow::FNodeParameters& InParam, FGuid InGuid)
+	: Super(InParam, InGuid)
+{
+	RegisterOutputConnection(&GroupedParameters);
+
+	// Create as many inputs as the node requires to have by default
+	AddDefaultInputs(InputMaterialParameters);
+}
+
+
+void FMakeMutableMaterialParametersArrayNode::Evaluate(UE::Dataflow::FContext& Context, const FDataflowOutput* Out) const
+{
+	EvaluateParameterNode(InputMaterialParameters, GroupedParameters, Context, Out);
+}
+
+
+TArray<UE::Dataflow::FPin> FMakeMutableMaterialParametersArrayNode::AddPins()
+{
+	return AddParameterPin(InputMaterialParameters);
+}
+
+
+bool FMakeMutableMaterialParametersArrayNode::CanRemovePin() const
+{
+	return CanRemoveParameterPin(InputMaterialParameters);
+}
+
+
+TArray<UE::Dataflow::FPin> FMakeMutableMaterialParametersArrayNode::GetPinsToRemove() const
+{
+	return GetParameterPinsToRemove(InputMaterialParameters);
+}
+
+
+void FMakeMutableMaterialParametersArrayNode::OnPinRemoved(const UE::Dataflow::FPin& Pin)
+{
+	OnParameterPinRemoved(InputMaterialParameters, Pin);
+}
+
+
+void FMakeMutableMaterialParametersArrayNode::PostSerialize(const FArchive& Ar)
+{
+	PostNodeSerialize(InputMaterialParameters, Ar);
+}
diff --git a/Engine/Plugins/Experimental/MutableDataflow/Source/MutableDataflowEditor/Private/Nodes/MakeMutableSkeletalMeshParametersArrayNode.cpp b/Engine/Plugins/Experimental/MutableDataflow/Source/MutableDataflowEditor/Private/Nodes/MakeMutableSkeletalMeshParametersArrayNode.cpp
new file mode 100644
--- /dev/null
+++ b/Engine/Plugins/Experimental/MutableDataflow/Source/MutableDataflowEditor/Private/Nodes/MakeMutableSkeletalMeshParametersArrayNode.cpp
@@ -0,0 +1,50 @@
+// Copyright Epic Games, Inc. All Rights Reserved.
+
+
+#include "Nodes/MakeMutableSkeletalMeshParametersArrayNode.h"
+
+
+FMakeMutableSkeletalMeshParametersArrayNode::FMakeMutableSkeletalMeshParametersArrayNode(const UE::Dataflow::FNodeParameters& InParam, FGuid InGuid)
+	: Super(InParam, InGuid)
+{
+	RegisterOutputConnection(&GroupedParameters);
+
+	// Create as many inputs as the node requires to have by default
+	AddDefaultInputs(InputSkeletalMeshParameters);
+}
+
+
+void FMakeMutableSkeletalMeshParametersArrayNode::Evaluate(UE::Dataflow::FContext& Context, const FDataflowOutput* Out) const
+{
+	EvaluateParameterNode(InputSkeletalMeshParameters, GroupedParameters, Context, Out);
+}
+
+
+TArray<UE::Dataflow::FPin> FMakeMutableSkeletalMeshParametersArrayNode::AddPins()
+{
+	return AddParameterPin(InputSkeletalMeshParameters);
+}
+
+
+bool FMakeMutableSkeletalMeshParametersArrayNode::CanRemovePin() const
+{
+	return CanRemoveParameterPin(InputSkeletalMeshParameters);
+}
+
+
+TArray<UE::Dataflow::FPin> FMakeMutableSkeletalMeshParametersArrayNode::GetPinsToRemove() const
+{
+	return GetParameterPinsToRemove(InputSkeletalMeshParameters);
+}
+
+
+void FMakeMutableSkeletalMeshParametersArrayNode::OnPinRemoved(const UE::Dataflow::FPin& Pin)
+{
+	OnParameterPinRemoved(InputSkeletalMeshParameters, Pin);
+}
+
+
+void FMakeMutableSkeletalMeshParametersArrayNode::PostSerialize(const FArchive& Ar)
+{
+	PostNodeSerialize(InputSkeletalMeshParameters, Ar);
+}
diff --git a/Engine/Plugins/Experimental/MutableDataflow/Source/MutableDataflowEditor/Private/Nodes/MakeMutableTextureParametersArrayNode.cpp b/Engine/Plugins/Experimental/MutableDataflow/Source/MutableDataflowEditor/Private/Nodes/MakeMutableTextureParametersArrayNode.cpp
new file mode 100644
--- /dev/null
+++ b/Engine/Plugins/Experimental/MutableDataflow/Source/MutableDataflowEditor/Private/Nodes/MakeMutableTextureParametersArrayNode.cpp
@@ -0,0 +1,50 @@
+﻿// Copyright Epic Games, Inc. All Rights Reserved.
+
+
+#include "Nodes/MakeMutableTextureParametersArrayNode.h"
+
+
+FMakeMutableTextureParametersArrayNode::FMakeMutableTextureParametersArrayNode(const UE::Dataflow::FNodeParameters& InParam, FGuid InGuid)
+	: Super(InParam, InGuid)
+{
+	RegisterOutputConnection(&GroupedParameters);
+
+	// Create as many inputs as the node requires to have by default
+	AddDefaultInputs(InputTextureParameters);
+}
+
+
+void FMakeMutableTextureParametersArrayNode::Evaluate(UE::Dataflow::FContext& Context, const FDataflowOutput* Out) const
+{
+	EvaluateParameterNode(InputTextureParameters, GroupedParameters, Context, Out);
+}
+
+
+TArray<UE::Dataflow::FPin> FMakeMutableTextureParametersArrayNode::AddPins()
+{
+	return AddParameterPin(InputTextureParameters);
+}
+
+
+bool FMakeMutableTextureParametersArrayNode::CanRemovePin() const
+{
+	return CanRemoveParameterPin(InputTextureParameters);
+}
+
+
+TArray<UE::Dataflow::FPin> FMakeMutableTextureParametersArrayNode::GetPinsToRemove() const
+{
+	return GetParameterPinsToRemove(InputTextureParameters);
+}
+
+
+void FMakeMutableTextureParametersArrayNode::OnPinRemoved(const UE::Dataflow::FPin& Pin)
+{
+	OnParameterPinRemoved(InputTextureParameters, Pin);
+}
+
+
+void FMakeMutableTextureParametersArrayNode::PostSerialize(const FArchive& Ar)
+{
+	PostNodeSerialize(InputTextureParameters, Ar);
+}
diff --git a/Engine/Plugins/Experimental/MutableDataflow/Source/MutableDataflowEditor/Private/Nodes/MutableMaterialParameterNode.cpp b/Engine/Plugins/Experimental/MutableDataflow/Source/MutableDataflowEditor/Private/Nodes/MutableMaterialParameterNode.cpp
new file mode 100644
--- /dev/null
+++ b/Engine/Plugins/Experimental/MutableDataflow/Source/MutableDataflowEditor/Private/Nodes/MutableMaterialParameterNode.cpp
@@ -0,0 +1,28 @@
+﻿// Copyright Epic Games, Inc. All Rights Reserved.
+
+#include "Nodes/MutableMaterialParameterNode.h"
+#include "Materials/MaterialInterface.h"
+
+FMutableMaterialParameterNode::FMutableMaterialParameterNode(const UE::Dataflow::FNodeParameters& InParam, FGuid InGuid)
+: FDataflowNode(InParam, InGuid)
+{
+	RegisterInputConnection(&ParameterName);
+	RegisterInputConnection(&Material);
+
+	RegisterOutputConnection(&MaterialParameter);
+}
+
+
+void FMutableMaterialParameterNode::Evaluate(UE::Dataflow::FContext& Context, const FDataflowOutput* Out) const
+{
+	FMutableMaterialParameter Output;
+	Output.Name = GetValue(Context, &ParameterName);
+	Output.Material = GetValue(Context, &Material);
+	
+	if (Output.Name.IsEmpty())
+	{
+		Context.Warning(TEXT("The generated Material Parameter is empty. Make sure the parameter has a name."), this);
+	}
+	
+	SetValue(Context, Output, &MaterialParameter);
+}
diff --git a/Engine/Plugins/Experimental/MutableDataflow/Source/MutableDataflowEditor/Private/Nodes/MutableSkeletalMeshParameterNode.cpp b/Engine/Plugins/Experimental/MutableDataflow/Source/MutableDataflowEditor/Private/Nodes/MutableSkeletalMeshParameterNode.cpp
new file mode 100644
--- /dev/null
+++ b/Engine/Plugins/Experimental/MutableDataflow/Source/MutableDataflowEditor/Private/Nodes/MutableSkeletalMeshParameterNode.cpp
@@ -0,0 +1,28 @@
+﻿// Copyright Epic Games, Inc. All Rights Reserved.
+
+#include "Nodes/MutableSkeletalMeshParameterNode.h"
+#include "Engine/SkeletalMesh.h"
+
+FMutableSkeletalMeshParameterNode::FMutableSkeletalMeshParameterNode(const UE::Dataflow::FNodeParameters& InParam, FGuid InGuid)
+	: FDataflowNode(InParam, InGuid)
+{
+	RegisterInputConnection(&ParameterName);
+	RegisterInputConnection(&SkeletalMesh);
+
+	RegisterOutputConnection(&SkeletalMeshParameter);
+}
+
+
+void FMutableSkeletalMeshParameterNode::Evaluate(UE::Dataflow::FContext& Context, const FDataflowOutput* Out) const
+{
+	FMutableSkeletalMeshParameter Output;
+	Output.Name = GetValue(Context, &ParameterName);
+	Output.Mesh = GetValue(Context, &SkeletalMesh);
+
+	if (Output.Name.IsEmpty())
+	{
+		Context.Warning(TEXT("The generated Skeletal Mesh Parameter is empty. Make sure the parameter has a name."), this);
+	}
+
+	SetValue(Context, Output, &SkeletalMeshParameter);
+}
diff --git a/Engine/Plugins/Experimental/MutableDataflow/Source/MutableDataflowEditor/Private/Nodes/MutableTextureParameterNode.cpp b/Engine/Plugins/Experimental/MutableDataflow/Source/MutableDataflowEditor/Private/Nodes/MutableTextureParameterNode.cpp
new file mode 100644
--- /dev/null
+++ b/Engine/Plugins/Experimental/MutableDataflow/Source/MutableDataflowEditor/Private/Nodes/MutableTextureParameterNode.cpp
@@ -0,0 +1,28 @@
+﻿// Copyright Epic Games, Inc. All Rights Reserved.
+
+#include "Nodes/MutableTextureParameterNode.h"
+#include "Engine/Texture2D.h"
+
+FMutableTextureParameterNode::FMutableTextureParameterNode(const UE::Dataflow::FNodeParameters& InParam, FGuid InGuid)
+	: FDataflowNode(InParam, InGuid)
+{
+	RegisterInputConnection(&ParameterName);
+	RegisterInputConnection(&Texture);
+
+	RegisterOutputConnection(&TextureParameter);
+}
+
+
+void FMutableTextureParameterNode::Evaluate(UE::Dataflow::FContext& Context, const FDataflowOutput* Out) const
+{
+	FMutableTextureParameter Output;
+	Output.Name =  GetValue(Context, &ParameterName);
+	Output.Texture = GetValue(Context, &Texture);
+	
+	if (Output.Name.IsEmpty())
+	{
+		Context.Warning(TEXT("The generated Texture Parameter is empty. Make sure the parameter has a name."), this);
+	}
+
+	SetValue(Context, Output, &TextureParameter);
+}
\ No newline at end of file
diff --git a/Engine/Plugins/Experimental/MutableDataflow/Source/MutableDataflowEditor/Public/MutableDataflowEditorModule.h b/Engine/Plugins/Experimental/MutableDataflow/Source/MutableDataflowEditor/Public/MutableDataflowEditorModule.h
new file mode 100644
--- /dev/null
+++ b/Engine/Plugins/Experimental/MutableDataflow/Source/MutableDataflowEditor/Public/MutableDataflowEditorModule.h
@@ -0,0 +1,9 @@
+// Copyright Epic Games, Inc. All Rights Reserved.
+#pragma once
+
+#include "Modules/ModuleInterface.h"
+
+class FMutableDataflowEditorModule : public IModuleInterface
+{
+	virtual void StartupModule() override;
+};
diff --git a/Engine/Plugins/Experimental/MutableDataflow/Source/MutableDataflowEditor/Public/MutableDataflowParameters.h b/Engine/Plugins/Experimental/MutableDataflow/Source/MutableDataflowEditor/Public/MutableDataflowParameters.h
new file mode 100644
--- /dev/null
+++ b/Engine/Plugins/Experimental/MutableDataflow/Source/MutableDataflowEditor/Public/MutableDataflowParameters.h
@@ -0,0 +1,56 @@
+// Copyright Epic Games, Inc. All Rights Reserved.
+#pragma once
+
+#include "CoreMinimal.h"
+
+#include "MutableDataflowParameters.generated.h"
+
+class USkeletalMesh;
+class UTexture2D;
+class UMaterialInterface;
+
+USTRUCT()
+struct FMutableParameterBase
+{
+	GENERATED_USTRUCT_BODY()
+
+	virtual ~FMutableParameterBase() = default;
+	
+	UPROPERTY(EditAnywhere, Category = "Mutable")
+	FString Name;
+	
+	bool operator== (const FMutableParameterBase& Other) const
+	{
+		return this->Name.ToLower() == Other.Name.ToLower();
+	}
+};
+
+
+USTRUCT(BlueprintType, meta = (Experimental))
+struct FMutableSkeletalMeshParameter : public FMutableParameterBase
+{
+	GENERATED_USTRUCT_BODY()
+	
+	UPROPERTY(EditAnywhere, Category = "Mutable")
+	TObjectPtr<USkeletalMesh> Mesh;
+};
+
+
+USTRUCT(BlueprintType, meta = (Experimental))
+struct FMutableTextureParameter : public FMutableParameterBase
+{
+	GENERATED_USTRUCT_BODY()
+	
+	UPROPERTY(EditAnywhere, Category = "Mutable")
+	TObjectPtr<UTexture2D> Texture;
+};
+
+
+USTRUCT(BlueprintType, meta = (Experimental))
+struct FMutableMaterialParameter : public FMutableParameterBase
+{
+	GENERATED_USTRUCT_BODY()
+	
+	UPROPERTY(EditAnywhere, Category = "Mutable")
+	TObjectPtr<UMaterialInterface> Material;
+};
\ No newline at end of file
diff --git a/Engine/Plugins/Experimental/MutableDataflow/Source/MutableDataflowEditor/Public/Nodes/COInstanceGeneratorNode.h b/Engine/Plugins/Experimental/MutableDataflow/Source/MutableDataflowEditor/Public/Nodes/COInstanceGeneratorNode.h
new file mode 100644
--- /dev/null
+++ b/Engine/Plugins/Experimental/MutableDataflow/Source/MutableDataflowEditor/Public/Nodes/COInstanceGeneratorNode.h
@@ -0,0 +1,153 @@
+// Copyright Epic Games, Inc. All Rights Reserved.
+#pragma once
+
+#include "Dataflow/DataflowNode.h"
+#include "Misc/Guid.h"
+#include "CoreMinimal.h"
+#include "Dataflow/DataflowFunctionProperty.h"
+#include "MuCO/CustomizableObject.h"
+#include "Engine/SkeletalMesh.h"
+#include "GeometryCollection/ManagedArrayCollection.h"
+#include "Dataflow/DataflowInputOutput.h"
+#include "MutableDataflowParameters.h"
+
+#include "COInstanceGeneratorNode.generated.h"
+
+USTRUCT(meta = (Experimental))
+struct FMutableGeneratedResource
+{
+	GENERATED_USTRUCT_BODY()
+
+	UPROPERTY(EditAnywhere, Category = "Mutable")
+	FString ComponentName;
+
+	UPROPERTY(EditAnywhere, Category = "Mutable")
+	TObjectPtr<USkeletalMesh> SkeletalMesh;
+};
+
+
+USTRUCT(meta = (Experimental))
+struct FCOInstanceGetComponentMesh : public FDataflowNode
+{
+	GENERATED_USTRUCT_BODY()
+	DATAFLOW_NODE_DEFINE_INTERNAL(FCOInstanceGetComponentMesh, "GetComponentMesh", "Mutable", "")
+
+	UPROPERTY(EditAnywhere, Category = "Mutable", meta = (DataflowInput))
+	FString ComponentName;
+
+	UPROPERTY(meta = (DataflowInput))
+	TArray<FMutableGeneratedResource> GeneratedResources;
+
+	UPROPERTY(meta = (DataflowOutput))
+	TObjectPtr<USkeletalMesh> ComponentSkeletalMesh;
+
+public:
+	
+	FCOInstanceGetComponentMesh(const UE::Dataflow::FNodeParameters& InParam, FGuid InGuid = FGuid::NewGuid());
+	
+	virtual void Evaluate(UE::Dataflow::FContext& Context, const FDataflowOutput* Out) const override;
+};
+
+
+/**
+ * Node designed to work as an entry point for mutable in the DataFlow graph.
+ * It allows the setup of a Customizable Object Instance and the extraction of the skeletal meshes generated by its update.
+ */
+USTRUCT(meta = (Experimental))
+struct FCOInstanceGeneratorNode : public FDataflowNode
+{
+	GENERATED_USTRUCT_BODY()
+	DATAFLOW_NODE_DEFINE_INTERNAL(FCOInstanceGeneratorNode, "GenerateCustomizableObjectInstance", "Mutable", "")
+
+public:
+	
+	FCOInstanceGeneratorNode(const UE::Dataflow::FNodeParameters& InParam, FGuid InGuid = FGuid::NewGuid());
+
+	virtual void OnPropertyChanged(UE::Dataflow::FContext& Context, const FPropertyChangedEvent& InPropertyChangedEvent) override;
+
+	virtual void Evaluate(UE::Dataflow::FContext& Context, const FDataflowOutput* Out) const override;
+
+private:
+	
+	// Inputs
+
+	UPROPERTY(EditAnywhere, Category = "Mutable", meta = (DataflowInput))
+	TObjectPtr<UCustomizableObject> CustomizableObject;
+
+	
+	UPROPERTY(EditAnywhere, Category = "Mutable Parameters", meta = (DataflowInput))
+	TArray<FMutableSkeletalMeshParameter> SkeletalMeshParameters;
+	TArray<FMutableSkeletalMeshParameter> CachedSkeletalMeshParameters;
+	
+	UPROPERTY(EditAnywhere, Category = "Mutable Parameters", meta = (DataflowInput))
+	TArray<FMutableTextureParameter> TextureParameters;
+	TArray<FMutableTextureParameter> CachedTextureParameters;
+
+	UPROPERTY(EditAnywhere, Category = "Mutable Parameters", meta = (DataflowInput))
+	TArray<FMutableMaterialParameter> MaterialParameters;
+	TArray<FMutableMaterialParameter> CachedMaterialParameters;
+
+	// Outputs
+
+	// First option : an array of structures that contain the pair of data Component name + Skeletal Mesh
+	UPROPERTY(meta = (DataflowOutput))
+	TArray<FMutableGeneratedResource> GeneratedResources;
+	
+	UPROPERTY(meta = (DataflowOutput))
+	TArray<TObjectPtr<USkeletalMesh>> GeneratedSkeletalMeshes;
+
+
+	// Instance generated during the runtime and used to generate the output resources
+	UPROPERTY(EditAnywhere, Category = "Mutable", meta = (DisplayName = "Generated COI"))
+	TObjectPtr<UCustomizableObjectInstance> CustomizableObjectInstance;
+
+	/** Press this button to trigger first a CO compilation and then a COI update. This will generate the resources needed by the outputs of the node. */
+	UPROPERTY(EditAnywhere, Category = "Mutable", meta = (DisplayName = "Generate Resources", ButtonImage = "Icons.Refresh"))
+	FDataflowFunctionProperty GenerateInstanceResources;
+	
+	
+	/** Variable that works as a semaphore to prevent the scheduling of multiple Instance Compilation + Updates*/
+	bool bIsGeneratingResources = false;
+	
+	/**
+	 * Callback method invoked when the "Generate outputs" button gets pressed. It will trigger the compilation and later update of the generated 
+	 Customizable Object Instance
+	 */
+	void OnGenerateInstanceResourcesRequested(UE::Dataflow::FContext& Context);
+	
+	/**
+	 * Evaluate the input nodes and add their data to our local caches so we can later access it without the dataflow context
+	 */
+	void CacheNodeInputs(UE::Dataflow::FContext& Context);
+	
+	/**
+	 * Apply the values in the cached parameters into the provided Instance
+	 * @param Instance The instance that will get its parameters updated
+	 */
+	void ApplyInstanceParameters(UCustomizableObjectInstance& Instance);
+
+	/**
+	 * Empty the caches holding the data provided by the evaluation of the input pins
+	 */
+	void ClearCachedParameters();
+	
+	/**
+	 * Request the compilation of the CO so we can later generate a Customizable Object Instance.
+	 */
+	void RequestCompilation(UE::Dataflow::FContext& Context);
+	
+	/** 
+	 * Request the Update of the given COI.
+	 */
+	void RequestUpdate(UCustomizableObjectInstance& Instance);
+
+	/**
+	 * Extract from the given Customizable Object Instance the generated skeletal meshes for each one of the components the instance has.
+	 * @param Instance The COI we want to get the data from
+	 * @return An array of pairs where the key represents the name of the component and the value the skeletal mesh on that component.
+	 */
+	TArray<FMutableGeneratedResource> GetInstanceGeneratedMeshes(const UCustomizableObjectInstance& Instance) const;
+};
+
+
+
diff --git a/Engine/Plugins/Experimental/MutableDataflow/Source/MutableDataflowEditor/Public/Nodes/MakeMutableMaterialParametersArrayNode.h b/Engine/Plugins/Experimental/MutableDataflow/Source/MutableDataflowEditor/Public/Nodes/MakeMutableMaterialParametersArrayNode.h
new file mode 100644
--- /dev/null
+++ b/Engine/Plugins/Experimental/MutableDataflow/Source/MutableDataflowEditor/Public/Nodes/MakeMutableMaterialParametersArrayNode.h
@@ -0,0 +1,39 @@
+﻿// Copyright Epic Games, Inc. All Rights Reserved.
+#pragma once
+
+#include "Misc/Guid.h"
+#include "CoreMinimal.h"
+#include "MakeMutableParametersArrayBaseNode.h"
+#include "Dataflow/DataflowNode.h"
+#include "Dataflow/DataflowInputOutput.h"
+#include "MutableDataflowParameters.h"
+
+#include "MakeMutableMaterialParametersArrayNode.generated.h"
+
+
+USTRUCT(meta = (Experimental))
+struct FMakeMutableMaterialParametersArrayNode : public FMakeMutableParametersArrayBaseNode
+{
+	GENERATED_USTRUCT_BODY()
+	DATAFLOW_NODE_DEFINE_INTERNAL(FMakeMutableMaterialParametersArrayNode, "MakeMutableMaterialParametersArray", "Mutable", "")
+
+private:
+	UPROPERTY()
+	TArray<FMutableMaterialParameter> InputMaterialParameters;
+	
+	UPROPERTY(meta = (DataflowOutput, DisplayName = "Material Parameters"))
+	TArray<FMutableMaterialParameter> GroupedParameters;
+
+public:
+	
+	FMakeMutableMaterialParametersArrayNode(const UE::Dataflow::FNodeParameters& InParam, FGuid InGuid = FGuid::NewGuid());
+	
+	virtual void Evaluate(UE::Dataflow::FContext& Context, const FDataflowOutput* Out) const override;
+
+	// Overrides to allow the addition and removal of optional inputs
+	virtual TArray<UE::Dataflow::FPin> AddPins() override;
+	virtual bool CanRemovePin() const override;
+	virtual TArray<UE::Dataflow::FPin> GetPinsToRemove() const override;
+	virtual void OnPinRemoved(const UE::Dataflow::FPin& Pin) override;
+	virtual void PostSerialize(const FArchive& Ar) override; 
+};
\ No newline at end of file
diff --git a/Engine/Plugins/Experimental/MutableDataflow/Source/MutableDataflowEditor/Public/Nodes/MakeMutableParametersArrayBaseNode.h b/Engine/Plugins/Experimental/MutableDataflow/Source/MutableDataflowEditor/Public/Nodes/MakeMutableParametersArrayBaseNode.h
new file mode 100644
--- /dev/null
+++ b/Engine/Plugins/Experimental/MutableDataflow/Source/MutableDataflowEditor/Public/Nodes/MakeMutableParametersArrayBaseNode.h
@@ -0,0 +1,157 @@
+﻿// Copyright Epic Games, Inc. All Rights Reserved.
+#pragma once
+
+#include "Misc/Guid.h"
+#include "CoreMinimal.h"
+#include "Dataflow/DataflowNode.h"
+#include "Dataflow/DataflowInputOutput.h"
+
+#include "MakeMutableParametersArrayBaseNode.generated.h"
+
+
+USTRUCT(meta = (Experimental))
+struct FMakeMutableParametersArrayBaseNode : public FDataflowNode
+{
+	GENERATED_USTRUCT_BODY()
+	DATAFLOW_NODE_DEFINE_INTERNAL(FMakeMutableParametersArrayBaseNode, "MakeParametersArray", "Mutable", "")
+	
+public:
+	FMakeMutableParametersArrayBaseNode(const UE::Dataflow::FNodeParameters& InParam, FGuid InGuid = FGuid::NewGuid());
+
+protected:
+
+	/**
+	 * The min amount of input pins the node can have. 
+	 */
+	static constexpr uint16 MinRequiredInputsCount = 1;
+
+	template<typename T>
+	void AddDefaultInputs(TArray<T>& InParametersArray);
+
+	template<typename T>
+	void EvaluateParameterNode(const TArray<T>& InputParameters, const TArray<T>& OutputParameters, UE::Dataflow::FContext& Context, const FDataflowOutput* Out) const;
+
+	template<typename T>
+	UE::Dataflow::TConnectionReference<T> GetConnectionReference(const TArray<T>& ParametersArray, int32 Index) const;
+
+	template<typename T>
+	TArray<UE::Dataflow::FPin> AddParameterPin(TArray<T>& InOutParametersArray);
+	
+	template<typename T>
+	bool CanRemoveParameterPin(const TArray<T>& ParametersArray) const;
+
+	template<typename T>
+	TArray<UE::Dataflow::FPin> GetParameterPinsToRemove(const TArray<T>& ParametersArray) const;;
+	
+	template<typename T>
+	void OnParameterPinRemoved(TArray<T>& InParametersArray, const UE::Dataflow::FPin& Pin);
+
+	template<typename T>
+	void PostNodeSerialize(const TArray<T>& ParametersArray, const FArchive& Ar);
+
+	// FDataflowNode
+	virtual bool CanAddPin() const override;
+};
+
+
+inline FMakeMutableParametersArrayBaseNode::FMakeMutableParametersArrayBaseNode(const UE::Dataflow::FNodeParameters& InParam, FGuid InGuid)
+: Super(InParam, InGuid)
+{
+}
+
+inline bool FMakeMutableParametersArrayBaseNode::CanAddPin() const
+{
+	return true;
+}
+
+
+template <typename T>
+void FMakeMutableParametersArrayBaseNode::AddDefaultInputs(TArray<T>& InParametersArray)
+{
+	for (int32 Index = 0; Index < MinRequiredInputsCount; ++Index)
+	{
+		AddParameterPin(InParametersArray);
+	}
+}
+
+
+template <typename T>
+void FMakeMutableParametersArrayBaseNode::EvaluateParameterNode(const TArray<T>& InputParameters, const TArray<T>& OutputParameters,  UE::Dataflow::FContext& Context,
+	const FDataflowOutput* Out) const
+{
+	// Grab the values from the inputs
+	TArray<T> Output;
+	for (const T& Parameter : InputParameters)
+	{
+		Output.Add(GetValue(Context, &Parameter));
+	}
+
+	// Provide the composed array as an output
+	SetValue(Context, Output, &OutputParameters);
+}
+
+
+template<typename T>
+UE::Dataflow::TConnectionReference<T> FMakeMutableParametersArrayBaseNode::GetConnectionReference(const TArray<T>& ParametersArray, int32 Index) const
+{
+	return { &ParametersArray[Index], Index, &ParametersArray };
+}
+
+
+template <typename T>
+TArray<UE::Dataflow::FPin> FMakeMutableParametersArrayBaseNode::AddParameterPin(TArray<T>& InOutParametersArray)
+{
+	const int32 Index = InOutParametersArray.AddDefaulted();
+	UE::Dataflow::TConnectionReference<T> ConnectionRef = GetConnectionReference(InOutParametersArray, Index);
+	const FDataflowInput& Input = RegisterInputArrayConnection(ConnectionRef);
+	return { { UE::Dataflow::FPin::EDirection::INPUT, Input.GetType(), Input.GetName() } };
+}
+
+
+template <typename T>
+bool FMakeMutableParametersArrayBaseNode::CanRemoveParameterPin(const TArray<T>& ParametersArray) const
+{
+	return ParametersArray.Num() > MinRequiredInputsCount;
+}
+
+
+template <typename T>
+TArray<UE::Dataflow::FPin> FMakeMutableParametersArrayBaseNode::GetParameterPinsToRemove(const TArray<T>& ParametersArray) const
+{
+	const int32 Index = ParametersArray.Num() - 1;
+	check(ParametersArray.IsValidIndex(Index));
+	if (const FDataflowInput* const Input = FindInput(GetConnectionReference(ParametersArray, Index)))
+	{
+		return { { UE::Dataflow::FPin::EDirection::INPUT, Input->GetType(), Input->GetName() } };
+	}
+	return Super::GetPinsToRemove();
+}
+
+
+template <typename T>
+void FMakeMutableParametersArrayBaseNode::OnParameterPinRemoved(TArray<T>& InParametersArray, const UE::Dataflow::FPin& Pin)
+{
+	const int32 Index = InParametersArray.Num() - 1;
+	check(InParametersArray.IsValidIndex(Index));
+#if DO_CHECK
+		const FDataflowInput* const Input = FindInput(GetConnectionReference(InParametersArray, Index));
+		check(Input);
+		check(Input->GetName() == Pin.Name);
+		check(Input->GetType() == Pin.Type);
+#endif
+	InParametersArray.SetNum(Index);
+}
+
+
+template <typename T>
+void FMakeMutableParametersArrayBaseNode::PostNodeSerialize(const TArray<T>& ParametersArray, const FArchive& Ar)
+{
+	if (Ar.IsLoading())
+	{
+		// Make sure the node does have the same amount of inputs as when saved
+		for (int32 Index = 0; Index < ParametersArray.Num(); ++Index)
+		{
+			FindOrRegisterInputArrayConnection(GetConnectionReference(ParametersArray, Index));
+		}
+	}
+}
diff --git a/Engine/Plugins/Experimental/MutableDataflow/Source/MutableDataflowEditor/Public/Nodes/MakeMutableSkeletalMeshParametersArrayNode.h b/Engine/Plugins/Experimental/MutableDataflow/Source/MutableDataflowEditor/Public/Nodes/MakeMutableSkeletalMeshParametersArrayNode.h
new file mode 100644
--- /dev/null
+++ b/Engine/Plugins/Experimental/MutableDataflow/Source/MutableDataflowEditor/Public/Nodes/MakeMutableSkeletalMeshParametersArrayNode.h
@@ -0,0 +1,39 @@
+// Copyright Epic Games, Inc. All Rights Reserved.
+#pragma once
+
+#include "Misc/Guid.h"
+#include "CoreMinimal.h"
+#include "MakeMutableParametersArrayBaseNode.h"
+#include "Dataflow/DataflowNode.h"
+#include "Dataflow/DataflowInputOutput.h"
+#include "MutableDataflowParameters.h"
+
+#include "MakeMutableSkeletalMeshParametersArrayNode.generated.h"
+
+
+USTRUCT(meta = (Experimental))
+struct FMakeMutableSkeletalMeshParametersArrayNode : public FMakeMutableParametersArrayBaseNode
+{
+	GENERATED_USTRUCT_BODY()
+	DATAFLOW_NODE_DEFINE_INTERNAL(FMakeMutableSkeletalMeshParametersArrayNode, "MakeMutableSkeletalMeshParametersArray", "Mutable", "")
+
+private:
+	UPROPERTY()
+	TArray<FMutableSkeletalMeshParameter> InputSkeletalMeshParameters;
+	
+	UPROPERTY(meta = (DataflowOutput, DisplayName = "Skeletal Mesh Parameters"))
+	TArray<FMutableSkeletalMeshParameter> GroupedParameters;
+
+public:
+	
+	FMakeMutableSkeletalMeshParametersArrayNode(const UE::Dataflow::FNodeParameters& InParam, FGuid InGuid = FGuid::NewGuid());
+	
+	virtual void Evaluate(UE::Dataflow::FContext& Context, const FDataflowOutput* Out) const override;
+	
+	// Overrides to allow the addition and removal of optional inputs
+	virtual TArray<UE::Dataflow::FPin> AddPins() override;
+	virtual bool CanRemovePin() const override;
+	virtual TArray<UE::Dataflow::FPin> GetPinsToRemove() const override;
+	virtual void OnPinRemoved(const UE::Dataflow::FPin& Pin) override;
+	virtual void PostSerialize(const FArchive& Ar) override; 
+};
\ No newline at end of file
diff --git a/Engine/Plugins/Experimental/MutableDataflow/Source/MutableDataflowEditor/Public/Nodes/MakeMutableTextureParametersArrayNode.h b/Engine/Plugins/Experimental/MutableDataflow/Source/MutableDataflowEditor/Public/Nodes/MakeMutableTextureParametersArrayNode.h
new file mode 100644
--- /dev/null
+++ b/Engine/Plugins/Experimental/MutableDataflow/Source/MutableDataflowEditor/Public/Nodes/MakeMutableTextureParametersArrayNode.h
@@ -0,0 +1,39 @@
+﻿// Copyright Epic Games, Inc. All Rights Reserved.
+#pragma once
+
+#include "Misc/Guid.h"
+#include "CoreMinimal.h"
+#include "MakeMutableParametersArrayBaseNode.h"
+#include "Dataflow/DataflowNode.h"
+#include "Dataflow/DataflowInputOutput.h"
+#include "MutableDataflowParameters.h"
+
+#include "MakeMutableTextureParametersArrayNode.generated.h"
+
+
+USTRUCT(meta = (Experimental))
+struct FMakeMutableTextureParametersArrayNode : public FMakeMutableParametersArrayBaseNode
+{
+	GENERATED_USTRUCT_BODY()
+	DATAFLOW_NODE_DEFINE_INTERNAL(FMakeMutableTextureParametersArrayNode, "MakeMutableTextureParametersArray", "Mutable", "")
+
+private:
+	UPROPERTY()
+	TArray<FMutableTextureParameter> InputTextureParameters;
+	
+	UPROPERTY(meta = (DataflowOutput,  DisplayName = "Texture Parameters"))
+	TArray<FMutableTextureParameter> GroupedParameters;
+
+public:
+	
+	FMakeMutableTextureParametersArrayNode(const UE::Dataflow::FNodeParameters& InParam, FGuid InGuid = FGuid::NewGuid());
+	
+	virtual void Evaluate(UE::Dataflow::FContext& Context, const FDataflowOutput* Out) const override;
+
+	// Overrides to allow the addition and removal of optional inputs
+	virtual TArray<UE::Dataflow::FPin> AddPins() override;
+	virtual bool CanRemovePin() const override;
+	virtual TArray<UE::Dataflow::FPin> GetPinsToRemove() const override;
+	virtual void OnPinRemoved(const UE::Dataflow::FPin& Pin) override;
+	virtual void PostSerialize(const FArchive& Ar) override; 
+};
\ No newline at end of file
diff --git a/Engine/Plugins/Experimental/MutableDataflow/Source/MutableDataflowEditor/Public/Nodes/MutableMaterialParameterNode.h b/Engine/Plugins/Experimental/MutableDataflow/Source/MutableDataflowEditor/Public/Nodes/MutableMaterialParameterNode.h
new file mode 100644
--- /dev/null
+++ b/Engine/Plugins/Experimental/MutableDataflow/Source/MutableDataflowEditor/Public/Nodes/MutableMaterialParameterNode.h
@@ -0,0 +1,31 @@
+﻿// Copyright Epic Games, Inc. All Rights Reserved.
+#pragma once
+
+#include "Dataflow/DataflowNode.h"
+#include "Misc/Guid.h"
+#include "CoreMinimal.h"
+#include "MutableDataflowParameters.h"
+
+#include "MutableMaterialParameterNode.generated.h"
+
+USTRUCT(meta = (Experimental))
+struct FMutableMaterialParameterNode : public FDataflowNode
+{
+	GENERATED_USTRUCT_BODY()
+	DATAFLOW_NODE_DEFINE_INTERNAL(FMutableMaterialParameterNode, "MutableMaterialParameter", "Mutable", "")
+
+	UPROPERTY(EditAnywhere, Category = "Mutable", meta = (DataflowInput))
+	FString ParameterName;
+
+	UPROPERTY(EditAnywhere, Category = "Mutable", meta = (DataflowInput))
+	TObjectPtr<UMaterialInterface> Material;
+
+	UPROPERTY(meta = (DataflowOutput))
+	FMutableMaterialParameter MaterialParameter;
+
+public:
+	
+	FMutableMaterialParameterNode(const UE::Dataflow::FNodeParameters& InParam, FGuid InGuid = FGuid::NewGuid());
+	
+	virtual void Evaluate(UE::Dataflow::FContext& Context, const FDataflowOutput* Out) const override;
+};
\ No newline at end of file
diff --git a/Engine/Plugins/Experimental/MutableDataflow/Source/MutableDataflowEditor/Public/Nodes/MutableSkeletalMeshParameterNode.h b/Engine/Plugins/Experimental/MutableDataflow/Source/MutableDataflowEditor/Public/Nodes/MutableSkeletalMeshParameterNode.h
new file mode 100644
--- /dev/null
+++ b/Engine/Plugins/Experimental/MutableDataflow/Source/MutableDataflowEditor/Public/Nodes/MutableSkeletalMeshParameterNode.h
@@ -0,0 +1,32 @@
+﻿// Copyright Epic Games, Inc. All Rights Reserved.
+#pragma once
+
+#include "Dataflow/DataflowNode.h"
+#include "MutableDataflowParameters.h"
+#include "Misc/Guid.h"
+#include "CoreMinimal.h"
+
+#include "MutableSkeletalMeshParameterNode.generated.h"
+
+
+USTRUCT(meta = (Experimental))
+struct FMutableSkeletalMeshParameterNode : public FDataflowNode
+{
+	GENERATED_USTRUCT_BODY()
+	DATAFLOW_NODE_DEFINE_INTERNAL(FMutableSkeletalMeshParameterNode, "MutableSkeletalMeshParameter", "Mutable", "")
+	
+	UPROPERTY(EditAnywhere, Category = "Mutable", meta = (DataflowInput))
+	FString ParameterName;
+
+	UPROPERTY(EditAnywhere, Category = "Mutable", meta = (DataflowInput))
+	TObjectPtr<USkeletalMesh> SkeletalMesh;
+	
+	UPROPERTY(meta = (DataflowOutput))
+	FMutableSkeletalMeshParameter SkeletalMeshParameter;
+
+public:
+	
+	FMutableSkeletalMeshParameterNode(const UE::Dataflow::FNodeParameters& InParam, FGuid InGuid = FGuid::NewGuid());
+	
+	virtual void Evaluate(UE::Dataflow::FContext& Context, const FDataflowOutput* Out) const override;
+};
\ No newline at end of file
diff --git a/Engine/Plugins/Experimental/MutableDataflow/Source/MutableDataflowEditor/Public/Nodes/MutableTextureParameterNode.h b/Engine/Plugins/Experimental/MutableDataflow/Source/MutableDataflowEditor/Public/Nodes/MutableTextureParameterNode.h
new file mode 100644
--- /dev/null
+++ b/Engine/Plugins/Experimental/MutableDataflow/Source/MutableDataflowEditor/Public/Nodes/MutableTextureParameterNode.h
@@ -0,0 +1,31 @@
+﻿// Copyright Epic Games, Inc. All Rights Reserved.
+#pragma once
+
+#include "Dataflow/DataflowNode.h"
+#include "Misc/Guid.h"
+#include "CoreMinimal.h"
+#include "MutableDataflowParameters.h"
+
+#include "MutableTextureParameterNode.generated.h"
+
+USTRUCT(meta = (Experimental))
+struct FMutableTextureParameterNode : public FDataflowNode
+{
+	GENERATED_USTRUCT_BODY()
+	DATAFLOW_NODE_DEFINE_INTERNAL(FMutableTextureParameterNode, "MutableTextureParameter", "Mutable", "")
+
+	UPROPERTY(EditAnywhere, Category = "Mutable", meta = (DataflowInput))
+	FString ParameterName;
+
+	UPROPERTY(EditAnywhere, Category = "Mutable", meta = (DataflowInput))
+	TObjectPtr<UTexture2D> Texture;
+	
+	UPROPERTY(meta = (DataflowOutput))
+	FMutableTextureParameter TextureParameter;
+	
+public:
+	
+	FMutableTextureParameterNode(const UE::Dataflow::FNodeParameters& InParam, FGuid InGuid = FGuid::NewGuid());
+
+	virtual void Evaluate(UE::Dataflow::FContext& Context, const FDataflowOutput* Out) const override;
+};
\ No newline at end of file
