[mutable] Fixed dataflow node issue causing the generation of COIs even after a compilation failure
- Fixed incorrect comment

#rnx

diff --git a/Engine/Plugins/Experimental/MutableDataflow/Source/MutableDataflowEditor/Private/Nodes/COInstanceGeneratorNode.cpp b/Engine/Plugins/Experimental/MutableDataflow/Source/MutableDataflowEditor/Private/Nodes/COInstanceGeneratorNode.cpp
--- a/Engine/Plugins/Experimental/MutableDataflow/Source/MutableDataflowEditor/Private/Nodes/COInstanceGeneratorNode.cpp
+++ b/Engine/Plugins/Experimental/MutableDataflow/Source/MutableDataflowEditor/Private/Nodes/COInstanceGeneratorNode.cpp
@@ -156,6 +156,14 @@ void FCOInstanceGeneratorNode::RequestCompilation(UE::Dataflow::FContext& Contex
 	{
 		FCOInstanceGeneratorNode* HostNode = static_cast<FCOInstanceGeneratorNode*>(&This.Get());
 		check(HostNode);
+
+		// Early out if the request failed
+		if (CallbackParams.bRequestFailed || CallbackParams.bErrors)
+		{
+			// Tell the system it is ok to generate resources again
+			HostNode->bIsGeneratingResources = false;
+			return;
+		}
 			
 		// First generate the instance
 		HostNode->CustomizableObjectInstance = HostNode->CustomizableObject->CreateInstance();
@@ -173,7 +181,7 @@ void FCOInstanceGeneratorNode::RequestCompilation(UE::Dataflow::FContext& Contex
 	CompilationParams.bAsync = true;
 	CompilationParams.CallbackNative = CompilationDelegate;
 	
-	// todo: Expose compilation errors and warnings to the Dataflow context
+	// todo: UE-313428 Expose compilation errors and warnings to the Dataflow context
 	CustomizableObject->Compile(CompilationParams);
 
 	// Prevent the scheduling or more compilations while the generation of resources is in process
diff --git a/Engine/Plugins/Experimental/MutableDataflow/Source/MutableDataflowEditor/Public/Nodes/COInstanceGeneratorNode.h b/Engine/Plugins/Experimental/MutableDataflow/Source/MutableDataflowEditor/Public/Nodes/COInstanceGeneratorNode.h
--- a/Engine/Plugins/Experimental/MutableDataflow/Source/MutableDataflowEditor/Public/Nodes/COInstanceGeneratorNode.h
+++ b/Engine/Plugins/Experimental/MutableDataflow/Source/MutableDataflowEditor/Public/Nodes/COInstanceGeneratorNode.h
@@ -144,7 +144,7 @@ private:
 	/**
 	 * Extract from the given Customizable Object Instance the generated skeletal meshes for each one of the components the instance has.
 	 * @param Instance The COI we want to get the data from
-	 * @return An array of pairs where the key represents the name of the component and the value the skeletal mesh on that component.
+	 * @return An array of structures composed by the component name and the skeletal mesh generated for said component..
 	 */
 	TArray<FMutableGeneratedResource> GetInstanceGeneratedMeshes(const UCustomizableObjectInstance& Instance) const;
 };
