[mutable-Dataflow] Set the MutableDataflow plugin disabled by default
- Expanded description to make sure the nodes are not treated as final versions of themselves.

#rnx

diff --git a/Engine/Plugins/Experimental/MutableDataflow/MutableDataflow.uplugin b/Engine/Plugins/Experimental/MutableDataflow/MutableDataflow.uplugin
--- a/Engine/Plugins/Experimental/MutableDataflow/MutableDataflow.uplugin
+++ b/Engine/Plugins/Experimental/MutableDataflow/MutableDataflow.uplugin
@@ -3,7 +3,7 @@
 	"Version": 1,
 	"VersionName": "0.1",
 	"FriendlyName": "Mutable Dataflow Extensions",
-	"Description": "Adds Mutable functionality to work with Dataflow objects from the Dataflow plugin",
+	"Description": "Adds Mutable functionality to work with Dataflow objects from the Dataflow plugin. WARNING: All nodes in this plugin are experimental and will be changed/deprecated in the future. The Dataflow system does not yet fully support mutable so the functionality of these nodes is currently quite limited and manual.",
 	"Category": "Animation",
 	"CreatedBy": "Epic Games, Inc.",
 	"CreatedByURL": "https://epicgames.com",
@@ -14,7 +14,7 @@
 	"IsBetaVersion": false,
 	"IsExperimentalVersion": true,
 	"Installed": false,
-	"EnabledByDefault": true,
+	"EnabledByDefault": false,
 	"Modules": 
 	[
 		{
