[mutable] Added multiple mutable parameter type supporting nodes to be used within Dataflow graphs
- Added Float Parameter Nodes (Param and Array filler nodes)
- Added Vector Parameter Nodes (Param and Array filler nodes)
- Added Enum Parameter Nodes (Param and Array filler nodes)
- Added Projector Parameter Nodes (Param and Array filler nodes)
- Added Transform Parameter Nodes (Param and Array filler nodes)
- Added Bool Parameter Nodes (Param and Array filler nodes)
- Some small tweaks in how the nodes are ordered and presented within DF.

diff --git a/Engine/Plugins/Experimental/MutableDataflow/Source/MutableDataflowEditor/Private/MutableDataflowEditorModule.cpp b/Engine/Plugins/Experimental/MutableDataflow/Source/MutableDataflowEditor/Private/MutableDataflowEditorModule.cpp
--- a/Engine/Plugins/Experimental/MutableDataflow/Source/MutableDataflowEditor/Private/MutableDataflowEditorModule.cpp
+++ b/Engine/Plugins/Experimental/MutableDataflow/Source/MutableDataflowEditor/Private/MutableDataflowEditorModule.cpp
@@ -2,21 +2,37 @@
 
 #include "MutableDataflowEditorModule.h"
 
+#include "Dataflow/ColorScheme.h"
 #include "Dataflow/DataflowAnyTypeRegistry.h"
 #include "Dataflow/DataflowNode.h"
 #include "Dataflow/DataflowCore.h"
 #include "Nodes/COInstanceGeneratorNode.h"
+#include "Nodes/MakeMutableBoolParametersArrayNode.h"
+#include "Nodes/MakeMutableEnumParametersArrayNode.h"
+#include "Nodes/MakeMutableFloatParametersArrayNode.h"
 #include "Nodes/MutableMaterialParameterNode.h"
 #include "Nodes/MutableSkeletalMeshParameterNode.h"
 #include "Nodes/MutableTextureParameterNode.h"
 #include "Nodes/MakeMutableMaterialParametersArrayNode.h"
+#include "Nodes/MakeMutableProjectorParametersArrayNode.h"
 #include "Nodes/MakeMutableSkeletalMeshParametersArrayNode.h"
 #include "Nodes/MakeMutableTextureParametersArrayNode.h"
+#include "Nodes/MakeMutableTransformParametersArrayNode.h"
+#include "Nodes/MakeMutableVectorParametersArrayNode.h"
+#include "Nodes/MutableBoolParameterNode.h"
+#include "Nodes/MutableEnumParameterNode.h"
+#include "Nodes/MutableFloatParameterNode.h"
+#include "Nodes/MutableProjectorParameterNode.h"
+#include "Nodes/MutableTransformParameterNode.h"
+#include "Nodes/MutableVectorParameterNode.h"
 
 IMPLEMENT_MODULE(FMutableDataflowEditorModule, MutableDataflowEditor);
 
 void FMutableDataflowEditorModule::StartupModule()
 {
+	const FLinearColor MutableNodeTitleColor = FColor(234, 255, 0);
+	DATAFLOW_NODE_REGISTER_CREATION_FACTORY_NODE_COLORS_BY_CATEGORY("Mutable", MutableNodeTitleColor, UE::Dataflow::FColorScheme::NodeBody);
+	
 	// Main nodes
 	DATAFLOW_NODE_REGISTER_CREATION_FACTORY(FCOInstanceGeneratorNode);
 	DATAFLOW_NODE_REGISTER_CREATION_FACTORY(FCOInstanceGetComponentMesh);
@@ -28,14 +44,32 @@ void FMutableDataflowEditorModule::StartupModule()
 	DATAFLOW_NODE_REGISTER_CREATION_FACTORY(FMutableSkeletalMeshParameterNode);
 	DATAFLOW_NODE_REGISTER_CREATION_FACTORY(FMutableTextureParameterNode);
 	DATAFLOW_NODE_REGISTER_CREATION_FACTORY(FMutableMaterialParameterNode);
-
+	DATAFLOW_NODE_REGISTER_CREATION_FACTORY(FMutableBoolParameterNode);
+	DATAFLOW_NODE_REGISTER_CREATION_FACTORY(FMutableFloatParameterNode);
+	DATAFLOW_NODE_REGISTER_CREATION_FACTORY(FMutableProjectorParameterNode);
+	DATAFLOW_NODE_REGISTER_CREATION_FACTORY(FMutableTransformParameterNode);
+	DATAFLOW_NODE_REGISTER_CREATION_FACTORY(FMutableVectorParameterNode);
+	DATAFLOW_NODE_REGISTER_CREATION_FACTORY(FMutableEnumParameterNode);
+	
 	// Parameter array generation nodes
 	DATAFLOW_NODE_REGISTER_CREATION_FACTORY(FMakeMutableSkeletalMeshParametersArrayNode);
 	DATAFLOW_NODE_REGISTER_CREATION_FACTORY(FMakeMutableMaterialParametersArrayNode);
 	DATAFLOW_NODE_REGISTER_CREATION_FACTORY(FMakeMutableTextureParametersArrayNode);
-
+	DATAFLOW_NODE_REGISTER_CREATION_FACTORY(FMakeMutableBoolParametersArrayNode);
+	DATAFLOW_NODE_REGISTER_CREATION_FACTORY(FMakeMutableFloatParametersArrayNode);
+	DATAFLOW_NODE_REGISTER_CREATION_FACTORY(FMakeMutableProjectorParametersArrayNode);
+	DATAFLOW_NODE_REGISTER_CREATION_FACTORY(FMakeMutableTransformParametersArrayNode);
+	DATAFLOW_NODE_REGISTER_CREATION_FACTORY(FMakeMutableVectorParametersArrayNode);
+	DATAFLOW_NODE_REGISTER_CREATION_FACTORY(FMakeMutableEnumParametersArrayNode);
+	
 	// Parameter node to array automatic conversion nodes
 	UE_DATAFLOW_REGISTER_AUTOCONVERT(FMutableSkeletalMeshParameter, TArray<FMutableSkeletalMeshParameter>, FMakeMutableSkeletalMeshParametersArrayNode);
 	UE_DATAFLOW_REGISTER_AUTOCONVERT(FMutableMaterialParameter, TArray<FMutableMaterialParameter>, FMakeMutableMaterialParametersArrayNode);
 	UE_DATAFLOW_REGISTER_AUTOCONVERT(FMutableTextureParameter, TArray<FMutableTextureParameter>, FMakeMutableTextureParametersArrayNode);
+	UE_DATAFLOW_REGISTER_AUTOCONVERT(FMutableBoolParameter, TArray<FMutableBoolParameter>, FMakeMutableBoolParametersArrayNode);
+	UE_DATAFLOW_REGISTER_AUTOCONVERT(FMutableFloatParameter, TArray<FMutableFloatParameter>, FMakeMutableFloatParametersArrayNode);
+	UE_DATAFLOW_REGISTER_AUTOCONVERT(FMutableProjectorParameter, TArray<FMutableProjectorParameter>, FMakeMutableProjectorParametersArrayNode);
+	UE_DATAFLOW_REGISTER_AUTOCONVERT(FMutableTransformParameter, TArray<FMutableTransformParameter>, FMakeMutableTransformParametersArrayNode);
+	UE_DATAFLOW_REGISTER_AUTOCONVERT(FMutableVectorParameter, TArray<FMutableVectorParameter>, FMakeMutableVectorParametersArrayNode);
+	UE_DATAFLOW_REGISTER_AUTOCONVERT(FMutableEnumParameter, TArray<FMutableEnumParameter>, FMakeMutableEnumParametersArrayNode);
 }
diff --git a/Engine/Plugins/Experimental/MutableDataflow/Source/MutableDataflowEditor/Private/Nodes/COInstanceGeneratorNode.cpp b/Engine/Plugins/Experimental/MutableDataflow/Source/MutableDataflowEditor/Private/Nodes/COInstanceGeneratorNode.cpp
--- a/Engine/Plugins/Experimental/MutableDataflow/Source/MutableDataflowEditor/Private/Nodes/COInstanceGeneratorNode.cpp
+++ b/Engine/Plugins/Experimental/MutableDataflow/Source/MutableDataflowEditor/Private/Nodes/COInstanceGeneratorNode.cpp
@@ -47,9 +47,16 @@ FCOInstanceGeneratorNode::FCOInstanceGeneratorNode(const UE::Dataflow::FNodePara
 {
 	// Inputs
 	RegisterInputConnection(&CustomizableObject);
+	
 	RegisterInputConnection(&SkeletalMeshParameters);
 	RegisterInputConnection(&TextureParameters);
 	RegisterInputConnection(&MaterialParameters);
+	RegisterInputConnection(&BoolParameters);
+	RegisterInputConnection(&EnumParameters);
+	RegisterInputConnection(&FloatParameters);
+	RegisterInputConnection(&VectorParameters);
+	RegisterInputConnection(&ProjectorParameters);
+	RegisterInputConnection(&TransformParameters);
 	
 	// Outputs
 	RegisterOutputConnection(&GeneratedResources);
@@ -120,6 +127,12 @@ void FCOInstanceGeneratorNode::ClearCachedParameters()
 	CachedSkeletalMeshParameters.Reset();
 	CachedMaterialParameters.Reset();
 	CachedTextureParameters.Reset();
+	CachedBoolParameters.Reset();
+	CachedEnumParameters.Reset();
+	CachedFloatParameters.Reset();
+	CachedVectorParameters.Reset();
+	CachedProjectorParameters.Reset();
+	CachedTransformParameters.Reset();
 }
 
 
@@ -196,52 +209,16 @@ void FCOInstanceGeneratorNode::CacheNodeInputs(UE::Dataflow::FContext& Context)
 	// Customizable Object
 	CustomizableObject = GetValue(Context, &CustomizableObject);
 
-	// Cache the parameters now that we have the context. Do not cache duplicate entries (based on the param name)
-	
-	// Texture Parameters
-	{
-		const TArray<FMutableTextureParameter>& TempTextureParametersArray = GetValue(Context, &TextureParameters);
-		for (const FMutableTextureParameter& TextureParameter : TempTextureParametersArray)
-		{
-			if (CachedTextureParameters.Contains(TextureParameter))
-			{
-				Context.Warning(FString::Printf( TEXT("Multiple Texture Parameters have been found sharing the \"%s\" name , only the first parameter will be used."), *TextureParameter.Name), this);
-				continue;
-			}
-		
-			CachedTextureParameters.Add(TextureParameter);
-		}
-	}
-	
-	// Mesh Parameters
-	{
-		const TArray<FMutableSkeletalMeshParameter>& TempMeshParametersArray = GetValue(Context, &SkeletalMeshParameters);
-		for (const FMutableSkeletalMeshParameter& MeshParameter : TempMeshParametersArray)
-		{
-			if (CachedSkeletalMeshParameters.Contains(MeshParameter))
-			{
-				Context.Warning(FString::Printf( TEXT("Multiple Skeletal Mesh Parameters have been found sharing the \"%s\" name , only the first parameter will be used."), *MeshParameter.Name), this);
-				continue;
-			}
-		
-			CachedSkeletalMeshParameters.Add(MeshParameter);
-		}
-	}
-	
-	// Material parameters
-	{
-		const TArray<FMutableMaterialParameter>& TempMaterialParametersArray = GetValue(Context, &MaterialParameters);
-		for (const FMutableMaterialParameter& MaterialParameter : TempMaterialParametersArray)
-		{
-			if (CachedMaterialParameters.Contains(MaterialParameter))
-			{
-				Context.Warning(FString::Printf( TEXT("Multiple Material Parameters have been found sharing the \"%s\" name , only the first parameter will be used."), *MaterialParameter.Name), this);
-				continue;
-			}
-		
-			CachedMaterialParameters.Add(MaterialParameter);
-		}
-	}
+	// Cache the parameters now that we have the context. Do not cache duplicate entries (based on the param name
+	CachedTextureParameters = GetPropertyParameters(&TextureParameters, TEXT("Texture"), Context);
+	CachedSkeletalMeshParameters = GetPropertyParameters(&SkeletalMeshParameters,TEXT("Skeletal Mesh"), Context);
+	CachedMaterialParameters = GetPropertyParameters(&MaterialParameters,TEXT("Material"), Context);
+	CachedBoolParameters = GetPropertyParameters(&BoolParameters,TEXT("Bool"), Context);
+	CachedEnumParameters = GetPropertyParameters(&EnumParameters,TEXT("Enum"), Context);
+	CachedFloatParameters = GetPropertyParameters(&FloatParameters,TEXT("Float"), Context);
+	CachedVectorParameters = GetPropertyParameters(&VectorParameters,TEXT("Vector"), Context);
+	CachedProjectorParameters = GetPropertyParameters(&ProjectorParameters,TEXT("Projector"), Context);
+	CachedTransformParameters = GetPropertyParameters(&TransformParameters,TEXT("Transform"),Context);
 }
 
 
@@ -251,7 +228,6 @@ void FCOInstanceGeneratorNode::ApplyInstanceParameters(UCustomizableObjectInstan
 	// In order to do that we will require to have async node evaluations and then access to the
 	// context to log the warning
 	
-	// Texture Parameters
 	for (const FMutableTextureParameter& TextureParameter : CachedTextureParameters)
 	{
 		const FString ParameterName = TextureParameter.Name;
@@ -261,7 +237,6 @@ void FCOInstanceGeneratorNode::ApplyInstanceParameters(UCustomizableObjectInstan
 		}
 	}
 	
-	// Mesh Parameters
 	for (const FMutableSkeletalMeshParameter& MeshParameter : CachedSkeletalMeshParameters)
 	{
 		const FString ParameterName = MeshParameter.Name;
@@ -271,7 +246,6 @@ void FCOInstanceGeneratorNode::ApplyInstanceParameters(UCustomizableObjectInstan
 		}
 	}
 	
-	// Material parameters
 	for (const FMutableMaterialParameter& MaterialParameter : CachedMaterialParameters)
 	{
 		const FString ParameterName = MaterialParameter.Name;
@@ -280,6 +254,60 @@ void FCOInstanceGeneratorNode::ApplyInstanceParameters(UCustomizableObjectInstan
 			Instance.SetMaterialParameterSelectedOption(ParameterName, MaterialParameter.Material);				
 		}
 	}
+	
+	for (const FMutableBoolParameter& BoolParameter : CachedBoolParameters)
+	{
+		const FString ParameterName = BoolParameter.Name;
+		if (Instance.ContainsBoolParameter(ParameterName))
+		{
+			Instance.SetBoolParameterSelectedOption(ParameterName, BoolParameter.Bool);				
+		}
+	}
+
+	for (const FMutableEnumParameter& EnumParameter : CachedEnumParameters)
+	{
+		const FString ParameterName = EnumParameter.Name;
+		if (Instance.ContainsEnumParameter(ParameterName))
+		{
+			Instance.SetEnumParameterSelectedOption(ParameterName, EnumParameter.OptionName);				
+		}
+	}
+
+	for (const FMutableFloatParameter& FloatParameter : CachedFloatParameters)
+	{
+		const FString ParameterName = FloatParameter.Name;
+		if (Instance.ContainsFloatParameter(ParameterName))
+		{
+			Instance.SetFloatParameterSelectedOption(ParameterName, FloatParameter.Float);				
+		}
+	}
+
+	for (const FMutableVectorParameter& VectorParameter : CachedVectorParameters)
+	{
+		const FString ParameterName = VectorParameter.Name;
+		if (Instance.ContainsVectorParameter(ParameterName))
+		{
+			Instance.SetVectorParameterSelectedOption(ParameterName, VectorParameter.Color);				
+		}
+	}
+	
+	for (const FMutableProjectorParameter& ProjectorParameter : CachedProjectorParameters)
+	{
+		const FString ParameterName = ProjectorParameter.Name;
+		if (Instance.ContainsProjectorParameter(ParameterName))
+		{
+			Instance.SetProjectorParameterSelectedOption(ParameterName, ProjectorParameter.Projector);
+		}
+	}
+	
+	for (const FMutableTransformParameter& TransformParameter : CachedTransformParameters)
+	{
+		const FString ParameterName = TransformParameter.Name;
+		if (Instance.ContainsTransformParameter(ParameterName))
+		{
+			Instance.SetTransformParameterSelectedOption(ParameterName, TransformParameter.Transform);				
+		}
+	}
 }
 
 
diff --git a/Engine/Plugins/Experimental/MutableDataflow/Source/MutableDataflowEditor/Private/Nodes/MakeMutableBoolParametersArrayNode.cpp b/Engine/Plugins/Experimental/MutableDataflow/Source/MutableDataflowEditor/Private/Nodes/MakeMutableBoolParametersArrayNode.cpp
new file mode 100644
--- /dev/null
+++ b/Engine/Plugins/Experimental/MutableDataflow/Source/MutableDataflowEditor/Private/Nodes/MakeMutableBoolParametersArrayNode.cpp
@@ -0,0 +1,50 @@
+﻿// Copyright Epic Games, Inc. All Rights Reserved.
+
+
+#include "Nodes/MakeMutableBoolParametersArrayNode.h"
+
+
+FMakeMutableBoolParametersArrayNode::FMakeMutableBoolParametersArrayNode(const UE::Dataflow::FNodeParameters& InParam, FGuid InGuid)
+	: Super(InParam, InGuid)
+{
+	RegisterOutputConnection(&GroupedParameters);
+
+	// Create as many inputs as the node requires to have by default
+	AddDefaultInputs(InputBoolParameters);
+}
+
+
+void FMakeMutableBoolParametersArrayNode::Evaluate(UE::Dataflow::FContext& Context, const FDataflowOutput* Out) const
+{
+	EvaluateParameterNode(InputBoolParameters, GroupedParameters, Context, Out);
+}
+
+
+TArray<UE::Dataflow::FPin> FMakeMutableBoolParametersArrayNode::AddPins()
+{
+	return AddParameterPin(InputBoolParameters);
+}
+
+
+bool FMakeMutableBoolParametersArrayNode::CanRemovePin() const
+{
+	return CanRemoveParameterPin(InputBoolParameters);
+}
+
+
+TArray<UE::Dataflow::FPin> FMakeMutableBoolParametersArrayNode::GetPinsToRemove() const
+{
+	return GetParameterPinsToRemove(InputBoolParameters);
+}
+
+
+void FMakeMutableBoolParametersArrayNode::OnPinRemoved(const UE::Dataflow::FPin& Pin)
+{
+	OnParameterPinRemoved(InputBoolParameters, Pin);
+}
+
+
+void FMakeMutableBoolParametersArrayNode::PostSerialize(const FArchive& Ar)
+{
+	PostNodeSerialize(InputBoolParameters, Ar);
+}
diff --git a/Engine/Plugins/Experimental/MutableDataflow/Source/MutableDataflowEditor/Private/Nodes/MakeMutableEnumParametersArrayNode.cpp b/Engine/Plugins/Experimental/MutableDataflow/Source/MutableDataflowEditor/Private/Nodes/MakeMutableEnumParametersArrayNode.cpp
new file mode 100644
--- /dev/null
+++ b/Engine/Plugins/Experimental/MutableDataflow/Source/MutableDataflowEditor/Private/Nodes/MakeMutableEnumParametersArrayNode.cpp
@@ -0,0 +1,50 @@
+﻿// Copyright Epic Games, Inc. All Rights Reserved.
+
+
+#include "Nodes/MakeMutableEnumParametersArrayNode.h"
+
+
+FMakeMutableEnumParametersArrayNode::FMakeMutableEnumParametersArrayNode(const UE::Dataflow::FNodeParameters& InParam, FGuid InGuid)
+	: Super(InParam, InGuid)
+{
+	RegisterOutputConnection(&GroupedParameters);
+
+	// Create as many inputs as the node requires to have by default
+	AddDefaultInputs(InputEnumParameters);
+}
+
+
+void FMakeMutableEnumParametersArrayNode::Evaluate(UE::Dataflow::FContext& Context, const FDataflowOutput* Out) const
+{
+	EvaluateParameterNode(InputEnumParameters, GroupedParameters, Context, Out);
+}
+
+
+TArray<UE::Dataflow::FPin> FMakeMutableEnumParametersArrayNode::AddPins()
+{
+	return AddParameterPin(InputEnumParameters);
+}
+
+
+bool FMakeMutableEnumParametersArrayNode::CanRemovePin() const
+{
+	return CanRemoveParameterPin(InputEnumParameters);
+}
+
+
+TArray<UE::Dataflow::FPin> FMakeMutableEnumParametersArrayNode::GetPinsToRemove() const
+{
+	return GetParameterPinsToRemove(InputEnumParameters);
+}
+
+
+void FMakeMutableEnumParametersArrayNode::OnPinRemoved(const UE::Dataflow::FPin& Pin)
+{
+	OnParameterPinRemoved(InputEnumParameters, Pin);
+}
+
+
+void FMakeMutableEnumParametersArrayNode::PostSerialize(const FArchive& Ar)
+{
+	PostNodeSerialize(InputEnumParameters, Ar);
+}
diff --git a/Engine/Plugins/Experimental/MutableDataflow/Source/MutableDataflowEditor/Private/Nodes/MakeMutableFloatParametersArrayNode.cpp b/Engine/Plugins/Experimental/MutableDataflow/Source/MutableDataflowEditor/Private/Nodes/MakeMutableFloatParametersArrayNode.cpp
new file mode 100644
--- /dev/null
+++ b/Engine/Plugins/Experimental/MutableDataflow/Source/MutableDataflowEditor/Private/Nodes/MakeMutableFloatParametersArrayNode.cpp
@@ -0,0 +1,50 @@
+﻿// Copyright Epic Games, Inc. All Rights Reserved.
+
+
+#include "Nodes/MakeMutableFloatParametersArrayNode.h"
+
+
+FMakeMutableFloatParametersArrayNode::FMakeMutableFloatParametersArrayNode(const UE::Dataflow::FNodeParameters& InParam, FGuid InGuid)
+	: Super(InParam, InGuid)
+{
+	RegisterOutputConnection(&GroupedParameters);
+
+	// Create as many inputs as the node requires to have by default
+	AddDefaultInputs(InputFloatParameters);
+}
+
+
+void FMakeMutableFloatParametersArrayNode::Evaluate(UE::Dataflow::FContext& Context, const FDataflowOutput* Out) const
+{
+	EvaluateParameterNode(InputFloatParameters, GroupedParameters, Context, Out);
+}
+
+
+TArray<UE::Dataflow::FPin> FMakeMutableFloatParametersArrayNode::AddPins()
+{
+	return AddParameterPin(InputFloatParameters);
+}
+
+
+bool FMakeMutableFloatParametersArrayNode::CanRemovePin() const
+{
+	return CanRemoveParameterPin(InputFloatParameters);
+}
+
+
+TArray<UE::Dataflow::FPin> FMakeMutableFloatParametersArrayNode::GetPinsToRemove() const
+{
+	return GetParameterPinsToRemove(InputFloatParameters);
+}
+
+
+void FMakeMutableFloatParametersArrayNode::OnPinRemoved(const UE::Dataflow::FPin& Pin)
+{
+	OnParameterPinRemoved(InputFloatParameters, Pin);
+}
+
+
+void FMakeMutableFloatParametersArrayNode::PostSerialize(const FArchive& Ar)
+{
+	PostNodeSerialize(InputFloatParameters, Ar);
+}
diff --git a/Engine/Plugins/Experimental/MutableDataflow/Source/MutableDataflowEditor/Private/Nodes/MakeMutableProjectorParametersArrayNode.cpp b/Engine/Plugins/Experimental/MutableDataflow/Source/MutableDataflowEditor/Private/Nodes/MakeMutableProjectorParametersArrayNode.cpp
new file mode 100644
--- /dev/null
+++ b/Engine/Plugins/Experimental/MutableDataflow/Source/MutableDataflowEditor/Private/Nodes/MakeMutableProjectorParametersArrayNode.cpp
@@ -0,0 +1,50 @@
+﻿// Copyright Epic Games, Inc. All Rights Reserved.
+
+
+#include "Nodes/MakeMutableProjectorParametersArrayNode.h"
+
+
+FMakeMutableProjectorParametersArrayNode::FMakeMutableProjectorParametersArrayNode(const UE::Dataflow::FNodeParameters& InParam, FGuid InGuid)
+	: Super(InParam, InGuid)
+{
+	RegisterOutputConnection(&GroupedParameters);
+
+	// Create as many inputs as the node requires to have by default
+	AddDefaultInputs(InputProjectorParameters);
+}
+
+
+void FMakeMutableProjectorParametersArrayNode::Evaluate(UE::Dataflow::FContext& Context, const FDataflowOutput* Out) const
+{
+	EvaluateParameterNode(InputProjectorParameters, GroupedParameters, Context, Out);
+}
+
+
+TArray<UE::Dataflow::FPin> FMakeMutableProjectorParametersArrayNode::AddPins()
+{
+	return AddParameterPin(InputProjectorParameters);
+}
+
+
+bool FMakeMutableProjectorParametersArrayNode::CanRemovePin() const
+{
+	return CanRemoveParameterPin(InputProjectorParameters);
+}
+
+
+TArray<UE::Dataflow::FPin> FMakeMutableProjectorParametersArrayNode::GetPinsToRemove() const
+{
+	return GetParameterPinsToRemove(InputProjectorParameters);
+}
+
+
+void FMakeMutableProjectorParametersArrayNode::OnPinRemoved(const UE::Dataflow::FPin& Pin)
+{
+	OnParameterPinRemoved(InputProjectorParameters, Pin);
+}
+
+
+void FMakeMutableProjectorParametersArrayNode::PostSerialize(const FArchive& Ar)
+{
+	PostNodeSerialize(InputProjectorParameters, Ar);
+}
diff --git a/Engine/Plugins/Experimental/MutableDataflow/Source/MutableDataflowEditor/Private/Nodes/MakeMutableTransformParametersArrayNode.cpp b/Engine/Plugins/Experimental/MutableDataflow/Source/MutableDataflowEditor/Private/Nodes/MakeMutableTransformParametersArrayNode.cpp
new file mode 100644
--- /dev/null
+++ b/Engine/Plugins/Experimental/MutableDataflow/Source/MutableDataflowEditor/Private/Nodes/MakeMutableTransformParametersArrayNode.cpp
@@ -0,0 +1,50 @@
+﻿// Copyright Epic Games, Inc. All Rights Reserved.
+
+
+#include "Nodes/MakeMutableTransformParametersArrayNode.h"
+
+
+FMakeMutableTransformParametersArrayNode::FMakeMutableTransformParametersArrayNode(const UE::Dataflow::FNodeParameters& InParam, FGuid InGuid)
+	: Super(InParam, InGuid)
+{
+	RegisterOutputConnection(&GroupedParameters);
+
+	// Create as many inputs as the node requires to have by default
+	AddDefaultInputs(InputTransformParameters);
+}
+
+
+void FMakeMutableTransformParametersArrayNode::Evaluate(UE::Dataflow::FContext& Context, const FDataflowOutput* Out) const
+{
+	EvaluateParameterNode(InputTransformParameters, GroupedParameters, Context, Out);
+}
+
+
+TArray<UE::Dataflow::FPin> FMakeMutableTransformParametersArrayNode::AddPins()
+{
+	return AddParameterPin(InputTransformParameters);
+}
+
+
+bool FMakeMutableTransformParametersArrayNode::CanRemovePin() const
+{
+	return CanRemoveParameterPin(InputTransformParameters);
+}
+
+
+TArray<UE::Dataflow::FPin> FMakeMutableTransformParametersArrayNode::GetPinsToRemove() const
+{
+	return GetParameterPinsToRemove(InputTransformParameters);
+}
+
+
+void FMakeMutableTransformParametersArrayNode::OnPinRemoved(const UE::Dataflow::FPin& Pin)
+{
+	OnParameterPinRemoved(InputTransformParameters, Pin);
+}
+
+
+void FMakeMutableTransformParametersArrayNode::PostSerialize(const FArchive& Ar)
+{
+	PostNodeSerialize(InputTransformParameters, Ar);
+}
diff --git a/Engine/Plugins/Experimental/MutableDataflow/Source/MutableDataflowEditor/Private/Nodes/MakeMutableVectorParametersArrayNode.cpp b/Engine/Plugins/Experimental/MutableDataflow/Source/MutableDataflowEditor/Private/Nodes/MakeMutableVectorParametersArrayNode.cpp
new file mode 100644
--- /dev/null
+++ b/Engine/Plugins/Experimental/MutableDataflow/Source/MutableDataflowEditor/Private/Nodes/MakeMutableVectorParametersArrayNode.cpp
@@ -0,0 +1,50 @@
+﻿// Copyright Epic Games, Inc. All Rights Reserved.
+
+
+#include "Nodes/MakeMutableVectorParametersArrayNode.h"
+
+
+FMakeMutableVectorParametersArrayNode::FMakeMutableVectorParametersArrayNode(const UE::Dataflow::FNodeParameters& InParam, FGuid InGuid)
+	: Super(InParam, InGuid)
+{
+	RegisterOutputConnection(&GroupedParameters);
+
+	// Create as many inputs as the node requires to have by default
+	AddDefaultInputs(InputVectorParameters);
+}
+
+
+void FMakeMutableVectorParametersArrayNode::Evaluate(UE::Dataflow::FContext& Context, const FDataflowOutput* Out) const
+{
+	EvaluateParameterNode(InputVectorParameters, GroupedParameters, Context, Out);
+}
+
+
+TArray<UE::Dataflow::FPin> FMakeMutableVectorParametersArrayNode::AddPins()
+{
+	return AddParameterPin(InputVectorParameters);
+}
+
+
+bool FMakeMutableVectorParametersArrayNode::CanRemovePin() const
+{
+	return CanRemoveParameterPin(InputVectorParameters);
+}
+
+
+TArray<UE::Dataflow::FPin> FMakeMutableVectorParametersArrayNode::GetPinsToRemove() const
+{
+	return GetParameterPinsToRemove(InputVectorParameters);
+}
+
+
+void FMakeMutableVectorParametersArrayNode::OnPinRemoved(const UE::Dataflow::FPin& Pin)
+{
+	OnParameterPinRemoved(InputVectorParameters, Pin);
+}
+
+
+void FMakeMutableVectorParametersArrayNode::PostSerialize(const FArchive& Ar)
+{
+	PostNodeSerialize(InputVectorParameters, Ar);
+}
diff --git a/Engine/Plugins/Experimental/MutableDataflow/Source/MutableDataflowEditor/Private/Nodes/MutableBoolParameterNode.cpp b/Engine/Plugins/Experimental/MutableDataflow/Source/MutableDataflowEditor/Private/Nodes/MutableBoolParameterNode.cpp
new file mode 100644
--- /dev/null
+++ b/Engine/Plugins/Experimental/MutableDataflow/Source/MutableDataflowEditor/Private/Nodes/MutableBoolParameterNode.cpp
@@ -0,0 +1,27 @@
+﻿// Copyright Epic Games, Inc. All Rights Reserved.
+
+#include "Nodes/MutableBoolParameterNode.h"
+
+#include "Dataflow/DataflowObject.h"
+
+
+FMutableBoolParameterNode::FMutableBoolParameterNode(const UE::Dataflow::FNodeParameters& InParam, FGuid InGuid)
+	: Super(InParam, InGuid)
+{
+	RegisterInputConnection(&Bool);
+
+	RegisterOutputConnection(&BoolParameter);
+}
+
+
+void FMutableBoolParameterNode::Evaluate(UE::Dataflow::FContext& Context, const FDataflowOutput* Out) const
+{
+	if (Out->IsA(&BoolParameter))
+	{
+		FMutableBoolParameter Output;
+		Output.Name = GetParameterName(Context);
+		Output.Bool = GetValue(Context, &Bool);
+	
+		SetValue(Context, Output, &BoolParameter);
+	}
+}
\ No newline at end of file
diff --git a/Engine/Plugins/Experimental/MutableDataflow/Source/MutableDataflowEditor/Private/Nodes/MutableEnumParameterNode.cpp b/Engine/Plugins/Experimental/MutableDataflow/Source/MutableDataflowEditor/Private/Nodes/MutableEnumParameterNode.cpp
new file mode 100644
--- /dev/null
+++ b/Engine/Plugins/Experimental/MutableDataflow/Source/MutableDataflowEditor/Private/Nodes/MutableEnumParameterNode.cpp
@@ -0,0 +1,26 @@
+﻿// Copyright Epic Games, Inc. All Rights Reserved.
+
+#include "Nodes/MutableEnumParameterNode.h"
+
+#include "Dataflow/DataflowObject.h"
+
+FMutableEnumParameterNode::FMutableEnumParameterNode(const UE::Dataflow::FNodeParameters& InParam, FGuid InGuid)
+	: Super(InParam, InGuid)
+{
+	RegisterInputConnection(&EnumOption);
+
+	RegisterOutputConnection(&EnumParameter);
+}
+
+
+void FMutableEnumParameterNode::Evaluate(UE::Dataflow::FContext& Context, const FDataflowOutput* Out) const
+{
+	if (Out->IsA(&EnumParameter))
+	{
+		FMutableEnumParameter Output;
+		Output.Name = GetParameterName(Context);
+		Output.OptionName = GetValue(Context, &EnumOption);
+	
+		SetValue(Context, Output, &EnumParameter);
+	}
+}
\ No newline at end of file
diff --git a/Engine/Plugins/Experimental/MutableDataflow/Source/MutableDataflowEditor/Private/Nodes/MutableFloatParameterNode.cpp b/Engine/Plugins/Experimental/MutableDataflow/Source/MutableDataflowEditor/Private/Nodes/MutableFloatParameterNode.cpp
new file mode 100644
--- /dev/null
+++ b/Engine/Plugins/Experimental/MutableDataflow/Source/MutableDataflowEditor/Private/Nodes/MutableFloatParameterNode.cpp
@@ -0,0 +1,26 @@
+﻿// Copyright Epic Games, Inc. All Rights Reserved.
+
+#include "Nodes/MutableFloatParameterNode.h"
+
+#include "Dataflow/DataflowObject.h"
+
+FMutableFloatParameterNode::FMutableFloatParameterNode(const UE::Dataflow::FNodeParameters& InParam, FGuid InGuid)
+	: Super(InParam, InGuid)
+{
+	RegisterInputConnection(&Float);
+
+	RegisterOutputConnection(&FloatParameter);
+}
+
+
+void FMutableFloatParameterNode::Evaluate(UE::Dataflow::FContext& Context, const FDataflowOutput* Out) const
+{
+	if (Out->IsA(&FloatParameter))
+	{
+		FMutableFloatParameter Output;
+		Output.Name = GetParameterName(Context);
+		Output.Float = GetValue(Context, &Float);
+	
+		SetValue(Context, Output, &FloatParameter);
+	}
+}
\ No newline at end of file
diff --git a/Engine/Plugins/Experimental/MutableDataflow/Source/MutableDataflowEditor/Private/Nodes/MutableMaterialParameterNode.cpp b/Engine/Plugins/Experimental/MutableDataflow/Source/MutableDataflowEditor/Private/Nodes/MutableMaterialParameterNode.cpp
--- a/Engine/Plugins/Experimental/MutableDataflow/Source/MutableDataflowEditor/Private/Nodes/MutableMaterialParameterNode.cpp
+++ b/Engine/Plugins/Experimental/MutableDataflow/Source/MutableDataflowEditor/Private/Nodes/MutableMaterialParameterNode.cpp
@@ -4,9 +4,8 @@
 #include "Materials/MaterialInterface.h"
 
 FMutableMaterialParameterNode::FMutableMaterialParameterNode(const UE::Dataflow::FNodeParameters& InParam, FGuid InGuid)
-: FDataflowNode(InParam, InGuid)
+: Super(InParam, InGuid)
 {
-	RegisterInputConnection(&ParameterName);
 	RegisterInputConnection(&Material);
 
 	RegisterOutputConnection(&MaterialParameter);
@@ -15,14 +14,12 @@ FMutableMaterialParameterNode::FMutableMaterialParameterNode(const UE::Dataflow:
 
 void FMutableMaterialParameterNode::Evaluate(UE::Dataflow::FContext& Context, const FDataflowOutput* Out) const
 {
-	FMutableMaterialParameter Output;
-	Output.Name = GetValue(Context, &ParameterName);
-	Output.Material = GetValue(Context, &Material);
-	
-	if (Output.Name.IsEmpty())
+	if (Out->IsA(&MaterialParameter))
 	{
-		Context.Warning(TEXT("The generated Material Parameter is empty. Make sure the parameter has a name."), this);
-	}
+		FMutableMaterialParameter Output;
+		Output.Name = GetParameterName(Context);
+		Output.Material = GetValue(Context, &Material);
 	
-	SetValue(Context, Output, &MaterialParameter);
+		SetValue(Context, Output, &MaterialParameter);
+	}
 }
diff --git a/Engine/Plugins/Experimental/MutableDataflow/Source/MutableDataflowEditor/Private/Nodes/MutableParameterNode.cpp b/Engine/Plugins/Experimental/MutableDataflow/Source/MutableDataflowEditor/Private/Nodes/MutableParameterNode.cpp
new file mode 100644
--- /dev/null
+++ b/Engine/Plugins/Experimental/MutableDataflow/Source/MutableDataflowEditor/Private/Nodes/MutableParameterNode.cpp
@@ -0,0 +1,22 @@
+﻿// Copyright Epic Games, Inc. All Rights Reserved.
+
+#include "Nodes/MutableParameterNode.h"
+#include "Dataflow/DataflowObject.h"
+
+FMutableParameterNode::FMutableParameterNode(const UE::Dataflow::FNodeParameters& InParam, FGuid InGuid)
+	: Super(InParam, InGuid)
+{
+	RegisterInputConnection(&ParameterName);
+}
+
+FString FMutableParameterNode::GetParameterName(UE::Dataflow::FContext& Context) const
+{
+	const FString OutParameterName = GetValue(Context, &ParameterName);
+	
+	if (OutParameterName.IsEmpty())
+	{
+		Context.Warning(TEXT("The generated Mutable Parameter name is empty."), this);
+	}
+	
+	return OutParameterName;
+}
\ No newline at end of file
diff --git a/Engine/Plugins/Experimental/MutableDataflow/Source/MutableDataflowEditor/Private/Nodes/MutableProjectorParameterNode.cpp b/Engine/Plugins/Experimental/MutableDataflow/Source/MutableDataflowEditor/Private/Nodes/MutableProjectorParameterNode.cpp
new file mode 100644
--- /dev/null
+++ b/Engine/Plugins/Experimental/MutableDataflow/Source/MutableDataflowEditor/Private/Nodes/MutableProjectorParameterNode.cpp
@@ -0,0 +1,27 @@
+﻿// Copyright Epic Games, Inc. All Rights Reserved.
+
+#include "Nodes/MutableProjectorParameterNode.h"
+
+#include "Dataflow/DataflowObject.h"
+
+
+FMutableProjectorParameterNode::FMutableProjectorParameterNode(const UE::Dataflow::FNodeParameters& InParam, FGuid InGuid)
+	: Super(InParam, InGuid)
+{
+	RegisterInputConnection(&Projector);
+
+	RegisterOutputConnection(&ProjectorParameter);
+}
+
+
+void FMutableProjectorParameterNode::Evaluate(UE::Dataflow::FContext& Context, const FDataflowOutput* Out) const
+{
+	if (Out->IsA(&ProjectorParameter))
+	{
+		FMutableProjectorParameter Output;
+		Output.Name = GetParameterName(Context);
+		Output.Projector = GetValue(Context, &Projector);
+	
+		SetValue(Context, Output, &ProjectorParameter);
+	}
+}
\ No newline at end of file
diff --git a/Engine/Plugins/Experimental/MutableDataflow/Source/MutableDataflowEditor/Private/Nodes/MutableSkeletalMeshParameterNode.cpp b/Engine/Plugins/Experimental/MutableDataflow/Source/MutableDataflowEditor/Private/Nodes/MutableSkeletalMeshParameterNode.cpp
--- a/Engine/Plugins/Experimental/MutableDataflow/Source/MutableDataflowEditor/Private/Nodes/MutableSkeletalMeshParameterNode.cpp
+++ b/Engine/Plugins/Experimental/MutableDataflow/Source/MutableDataflowEditor/Private/Nodes/MutableSkeletalMeshParameterNode.cpp
@@ -1,12 +1,13 @@
 ﻿// Copyright Epic Games, Inc. All Rights Reserved.
 
 #include "Nodes/MutableSkeletalMeshParameterNode.h"
+
+#include "Dataflow/DataflowObject.h"
 #include "Engine/SkeletalMesh.h"
 
 FMutableSkeletalMeshParameterNode::FMutableSkeletalMeshParameterNode(const UE::Dataflow::FNodeParameters& InParam, FGuid InGuid)
-	: FDataflowNode(InParam, InGuid)
+	: Super(InParam, InGuid)
 {
-	RegisterInputConnection(&ParameterName);
 	RegisterInputConnection(&SkeletalMesh);
 
 	RegisterOutputConnection(&SkeletalMeshParameter);
@@ -15,14 +16,13 @@ FMutableSkeletalMeshParameterNode::FMutableSkeletalMeshParameterNode(const UE::D
 
 void FMutableSkeletalMeshParameterNode::Evaluate(UE::Dataflow::FContext& Context, const FDataflowOutput* Out) const
 {
-	FMutableSkeletalMeshParameter Output;
-	Output.Name = GetValue(Context, &ParameterName);
-	Output.Mesh = GetValue(Context, &SkeletalMesh);
-
-	if (Output.Name.IsEmpty())
+	if (Out->IsA(&SkeletalMeshParameter))
 	{
-		Context.Warning(TEXT("The generated Skeletal Mesh Parameter is empty. Make sure the parameter has a name."), this);
+		FMutableSkeletalMeshParameter Output;
+		Output.Name = GetParameterName(Context);
+		Output.Mesh = GetValue(Context, &SkeletalMesh);
+	
+		SetValue(Context, Output, &SkeletalMeshParameter);
 	}
-
-	SetValue(Context, Output, &SkeletalMeshParameter);
 }
+
diff --git a/Engine/Plugins/Experimental/MutableDataflow/Source/MutableDataflowEditor/Private/Nodes/MutableTextureParameterNode.cpp b/Engine/Plugins/Experimental/MutableDataflow/Source/MutableDataflowEditor/Private/Nodes/MutableTextureParameterNode.cpp
--- a/Engine/Plugins/Experimental/MutableDataflow/Source/MutableDataflowEditor/Private/Nodes/MutableTextureParameterNode.cpp
+++ b/Engine/Plugins/Experimental/MutableDataflow/Source/MutableDataflowEditor/Private/Nodes/MutableTextureParameterNode.cpp
@@ -4,9 +4,8 @@
 #include "Engine/Texture2D.h"
 
 FMutableTextureParameterNode::FMutableTextureParameterNode(const UE::Dataflow::FNodeParameters& InParam, FGuid InGuid)
-	: FDataflowNode(InParam, InGuid)
+	: Super(InParam, InGuid)
 {
-	RegisterInputConnection(&ParameterName);
 	RegisterInputConnection(&Texture);
 
 	RegisterOutputConnection(&TextureParameter);
@@ -15,14 +14,12 @@ FMutableTextureParameterNode::FMutableTextureParameterNode(const UE::Dataflow::F
 
 void FMutableTextureParameterNode::Evaluate(UE::Dataflow::FContext& Context, const FDataflowOutput* Out) const
 {
-	FMutableTextureParameter Output;
-	Output.Name =  GetValue(Context, &ParameterName);
-	Output.Texture = GetValue(Context, &Texture);
-	
-	if (Output.Name.IsEmpty())
+	if (Out->IsA(&TextureParameter))
 	{
-		Context.Warning(TEXT("The generated Texture Parameter is empty. Make sure the parameter has a name."), this);
-	}
+		FMutableTextureParameter Output;
+		Output.Name =  GetParameterName(Context);
+		Output.Texture = GetValue(Context, &Texture);
 
-	SetValue(Context, Output, &TextureParameter);
+		SetValue(Context, Output, &TextureParameter);
+	}
 }
\ No newline at end of file
diff --git a/Engine/Plugins/Experimental/MutableDataflow/Source/MutableDataflowEditor/Private/Nodes/MutableTransformParameterNode.cpp b/Engine/Plugins/Experimental/MutableDataflow/Source/MutableDataflowEditor/Private/Nodes/MutableTransformParameterNode.cpp
new file mode 100644
--- /dev/null
+++ b/Engine/Plugins/Experimental/MutableDataflow/Source/MutableDataflowEditor/Private/Nodes/MutableTransformParameterNode.cpp
@@ -0,0 +1,27 @@
+﻿// Copyright Epic Games, Inc. All Rights Reserved.
+
+#include "Nodes/MutableTransformParameterNode.h"
+
+#include "Dataflow/DataflowObject.h"
+
+
+FMutableTransformParameterNode::FMutableTransformParameterNode(const UE::Dataflow::FNodeParameters& InParam, FGuid InGuid)
+	: Super(InParam, InGuid)
+{
+	RegisterInputConnection(&Transform);
+
+	RegisterOutputConnection(&TransformParameter);
+}
+
+
+void FMutableTransformParameterNode::Evaluate(UE::Dataflow::FContext& Context, const FDataflowOutput* Out) const
+{
+	if (Out->IsA(&TransformParameter))
+	{
+		FMutableTransformParameter Output;
+		Output.Name = GetParameterName(Context);
+		Output.Transform = GetValue(Context, &Transform);
+	
+		SetValue(Context, Output, &TransformParameter);
+	}
+}
\ No newline at end of file
diff --git a/Engine/Plugins/Experimental/MutableDataflow/Source/MutableDataflowEditor/Private/Nodes/MutableVectorParameterNode.cpp b/Engine/Plugins/Experimental/MutableDataflow/Source/MutableDataflowEditor/Private/Nodes/MutableVectorParameterNode.cpp
new file mode 100644
--- /dev/null
+++ b/Engine/Plugins/Experimental/MutableDataflow/Source/MutableDataflowEditor/Private/Nodes/MutableVectorParameterNode.cpp
@@ -0,0 +1,27 @@
+﻿// Copyright Epic Games, Inc. All Rights Reserved.
+
+#include "Nodes/MutableVectorParameterNode.h"
+
+#include "Dataflow/DataflowObject.h"
+
+
+FMutableVectorParameterNode::FMutableVectorParameterNode(const UE::Dataflow::FNodeParameters& InParam, FGuid InGuid)
+: Super(InParam, InGuid)
+{
+	RegisterInputConnection(&Color);
+
+	RegisterOutputConnection(&VectorParameter);
+}
+
+
+void FMutableVectorParameterNode::Evaluate(UE::Dataflow::FContext& Context, const FDataflowOutput* Out) const
+{
+	if (Out->IsA(&VectorParameter))
+	{
+		FMutableVectorParameter Output;
+		Output.Name = GetParameterName(Context);
+		Output.Color = GetValue(Context, &Color);
+	
+		SetValue(Context, Output, &VectorParameter);
+	}
+}
\ No newline at end of file
diff --git a/Engine/Plugins/Experimental/MutableDataflow/Source/MutableDataflowEditor/Public/MutableDataflowParameters.h b/Engine/Plugins/Experimental/MutableDataflow/Source/MutableDataflowEditor/Public/MutableDataflowParameters.h
--- a/Engine/Plugins/Experimental/MutableDataflow/Source/MutableDataflowEditor/Public/MutableDataflowParameters.h
+++ b/Engine/Plugins/Experimental/MutableDataflow/Source/MutableDataflowEditor/Public/MutableDataflowParameters.h
@@ -2,6 +2,7 @@
 #pragma once
 
 #include "CoreMinimal.h"
+#include "MuCO/CustomizableObjectParameterTypeDefinitions.h"
 
 #include "MutableDataflowParameters.generated.h"
 
@@ -53,4 +54,65 @@ struct FMutableMaterialParameter : public FMutableParameterBase
 	
 	UPROPERTY(EditAnywhere, Category = "Mutable")
 	TObjectPtr<UMaterialInterface> Material;
-};
\ No newline at end of file
+};
+
+
+USTRUCT(BlueprintType, meta = (Experimental))
+struct FMutableBoolParameter : public FMutableParameterBase
+{
+	GENERATED_USTRUCT_BODY()
+	
+	UPROPERTY(EditAnywhere, Category = "Mutable")
+	bool Bool;
+};
+
+
+USTRUCT(BlueprintType, meta = (Experimental))
+struct FMutableEnumParameter : public FMutableParameterBase
+{
+	GENERATED_USTRUCT_BODY()
+	
+	UPROPERTY(EditAnywhere, Category = "Mutable")
+	FString OptionName;
+};
+
+
+USTRUCT(BlueprintType, meta = (Experimental))
+struct FMutableFloatParameter : public FMutableParameterBase
+{
+	GENERATED_USTRUCT_BODY()
+	
+	UPROPERTY(EditAnywhere, Category = "Mutable")
+	float Float;
+};
+
+
+USTRUCT(BlueprintType, meta = (Experimental))
+struct FMutableVectorParameter : public FMutableParameterBase
+{
+	GENERATED_USTRUCT_BODY()
+	
+	UPROPERTY(EditAnywhere, Category = "Mutable")
+	FLinearColor Color;
+};
+
+
+USTRUCT(BlueprintType, meta = (Experimental))
+struct FMutableProjectorParameter : public FMutableParameterBase
+{
+	GENERATED_USTRUCT_BODY()
+
+	UPROPERTY(EditAnywhere, Category = "Mutable")
+	FCustomizableObjectProjector Projector;
+};
+
+
+USTRUCT(BlueprintType, meta = (Experimental))
+struct FMutableTransformParameter : public FMutableParameterBase
+{
+	GENERATED_USTRUCT_BODY()
+	
+	UPROPERTY(EditAnywhere, Category = "Mutable")
+	FTransform Transform;
+};
+
diff --git a/Engine/Plugins/Experimental/MutableDataflow/Source/MutableDataflowEditor/Public/Nodes/COInstanceGeneratorNode.h b/Engine/Plugins/Experimental/MutableDataflow/Source/MutableDataflowEditor/Public/Nodes/COInstanceGeneratorNode.h
--- a/Engine/Plugins/Experimental/MutableDataflow/Source/MutableDataflowEditor/Public/Nodes/COInstanceGeneratorNode.h
+++ b/Engine/Plugins/Experimental/MutableDataflow/Source/MutableDataflowEditor/Public/Nodes/COInstanceGeneratorNode.h
@@ -87,6 +87,30 @@ private:
 	TArray<FMutableMaterialParameter> MaterialParameters;
 	TArray<FMutableMaterialParameter> CachedMaterialParameters;
 
+	UPROPERTY(EditAnywhere, Category = "Mutable Parameters", meta = (DataflowInput))
+	TArray<FMutableBoolParameter> BoolParameters;
+	TArray<FMutableBoolParameter> CachedBoolParameters;
+
+	UPROPERTY(EditAnywhere, Category = "Mutable Parameters", meta = (DataflowInput))
+	TArray<FMutableEnumParameter> EnumParameters;
+	TArray<FMutableEnumParameter> CachedEnumParameters;
+
+	UPROPERTY(EditAnywhere, Category = "Mutable Parameters", meta = (DataflowInput))
+	TArray<FMutableFloatParameter> FloatParameters;
+	TArray<FMutableFloatParameter> CachedFloatParameters;
+
+	UPROPERTY(EditAnywhere, Category = "Mutable Parameters", meta = (DataflowInput))
+	TArray<FMutableVectorParameter> VectorParameters;
+	TArray<FMutableVectorParameter> CachedVectorParameters;
+
+	UPROPERTY(EditAnywhere, Category = "Mutable Parameters", meta = (DataflowInput))
+	TArray<FMutableProjectorParameter> ProjectorParameters;
+	TArray<FMutableProjectorParameter> CachedProjectorParameters;
+
+	UPROPERTY(EditAnywhere, Category = "Mutable Parameters", meta = (DataflowInput))
+	TArray<FMutableTransformParameter> TransformParameters;
+	TArray<FMutableTransformParameter> CachedTransformParameters;
+	
 	// Outputs
 
 	// First option : an array of structures that contain the pair of data Component name + Skeletal Mesh
@@ -114,6 +138,39 @@ private:
 	 Customizable Object Instance
 	 */
 	void OnGenerateInstanceResourcesRequested(UE::Dataflow::FContext& Context);
+
+
+	/**
+	 * Get all the unique parameters defined for a given property.
+	 * @tparam T The type of parameter. EX : FMutableProjectorParameter, FMutableMaterialParameter ...
+	 * @param InParametersArrayProperty The property used as source to find the parameters.
+	 * @param ParameterTypeName The descriptive name of the parameter type. Used for logging.
+	 * @param Context The Dataflow context.
+	 * @return An array with all the parameters found. Only the first occurrence of each parameter will be placed there.
+	 */
+	template<typename T>
+	TArray<T> GetPropertyParameters(const TArray<T>* InParametersArrayProperty, const FString& ParameterTypeName, UE::Dataflow::FContext& Context)
+	{
+		check(InParametersArrayProperty);
+		
+		const TArray<T>& FoundParametersCandidates = GetValue(Context, InParametersArrayProperty);
+		
+		TArray<T> OutFoundParameters;
+		OutFoundParameters.Reserve(FoundParametersCandidates.Num());
+		
+		for (const T& CandidateParameter : FoundParametersCandidates)
+		{
+			if (OutFoundParameters.Contains(CandidateParameter))
+			{
+				Context.Warning(FString::Printf( TEXT("Multiple %s Parameters have been found sharing the \"%s\" name , only the first parameter will be used."), *ParameterTypeName, *CandidateParameter.Name), this);
+				continue;
+			}
+		
+			OutFoundParameters.Add(CandidateParameter);
+		}
+		
+		return OutFoundParameters;
+	}
 	
 	/**
 	 * Evaluate the input nodes and add their data to our local caches so we can later access it without the dataflow context
diff --git a/Engine/Plugins/Experimental/MutableDataflow/Source/MutableDataflowEditor/Public/Nodes/MakeMutableBoolParametersArrayNode.h b/Engine/Plugins/Experimental/MutableDataflow/Source/MutableDataflowEditor/Public/Nodes/MakeMutableBoolParametersArrayNode.h
new file mode 100644
--- /dev/null
+++ b/Engine/Plugins/Experimental/MutableDataflow/Source/MutableDataflowEditor/Public/Nodes/MakeMutableBoolParametersArrayNode.h
@@ -0,0 +1,40 @@
+﻿// Copyright Epic Games, Inc. All Rights Reserved.
+#pragma once
+
+#include "Misc/Guid.h"
+#include "CoreMinimal.h"
+#include "MakeMutableParametersArrayBaseNode.h"
+#include "Dataflow/DataflowNode.h"
+#include "Dataflow/DataflowInputOutput.h"
+#include "MutableDataflowParameters.h"
+
+#include "MakeMutableBoolParametersArrayNode.generated.h"
+
+
+
+USTRUCT(meta = (Experimental))
+struct FMakeMutableBoolParametersArrayNode : public FMakeMutableParametersArrayBaseNode
+{
+	GENERATED_USTRUCT_BODY()
+	DATAFLOW_NODE_DEFINE_INTERNAL(FMakeMutableBoolParametersArrayNode, "MakeMutableBoolParametersArray", "Mutable|Parameters", "")
+
+private:
+	UPROPERTY()
+	TArray<FMutableBoolParameter> InputBoolParameters;
+	
+	UPROPERTY(meta = (DataflowOutput, DisplayName = "Bool Parameters"))
+	TArray<FMutableBoolParameter> GroupedParameters;
+
+public:
+	
+	FMakeMutableBoolParametersArrayNode(const UE::Dataflow::FNodeParameters& InParam, FGuid InGuid = FGuid::NewGuid());
+	
+	virtual void Evaluate(UE::Dataflow::FContext& Context, const FDataflowOutput* Out) const override;
+	
+	// Overrides to allow the addition and removal of optional inputs
+	virtual TArray<UE::Dataflow::FPin> AddPins() override;
+	virtual bool CanRemovePin() const override;
+	virtual TArray<UE::Dataflow::FPin> GetPinsToRemove() const override;
+	virtual void OnPinRemoved(const UE::Dataflow::FPin& Pin) override;
+	virtual void PostSerialize(const FArchive& Ar) override; 
+};
\ No newline at end of file
diff --git a/Engine/Plugins/Experimental/MutableDataflow/Source/MutableDataflowEditor/Public/Nodes/MakeMutableEnumParametersArrayNode.h b/Engine/Plugins/Experimental/MutableDataflow/Source/MutableDataflowEditor/Public/Nodes/MakeMutableEnumParametersArrayNode.h
new file mode 100644
--- /dev/null
+++ b/Engine/Plugins/Experimental/MutableDataflow/Source/MutableDataflowEditor/Public/Nodes/MakeMutableEnumParametersArrayNode.h
@@ -0,0 +1,39 @@
+﻿// Copyright Epic Games, Inc. All Rights Reserved.
+#pragma once
+
+#include "Misc/Guid.h"
+#include "CoreMinimal.h"
+#include "MakeMutableParametersArrayBaseNode.h"
+#include "Dataflow/DataflowNode.h"
+#include "Dataflow/DataflowInputOutput.h"
+#include "MutableDataflowParameters.h"
+
+#include "MakeMutableEnumParametersArrayNode.generated.h"
+
+
+USTRUCT(meta = (Experimental))
+struct FMakeMutableEnumParametersArrayNode : public FMakeMutableParametersArrayBaseNode
+{
+	GENERATED_USTRUCT_BODY()
+	DATAFLOW_NODE_DEFINE_INTERNAL(FMakeMutableEnumParametersArrayNode, "MakeMutableEnumParametersArray", "Mutable|Parameters", "")
+
+private:
+	UPROPERTY()
+	TArray<FMutableEnumParameter> InputEnumParameters;
+	
+	UPROPERTY(meta = (DataflowOutput, DisplayName = "Enum Parameters"))
+	TArray<FMutableEnumParameter> GroupedParameters;
+
+public:
+	
+	FMakeMutableEnumParametersArrayNode(const UE::Dataflow::FNodeParameters& InParam, FGuid InGuid = FGuid::NewGuid());
+	
+	virtual void Evaluate(UE::Dataflow::FContext& Context, const FDataflowOutput* Out) const override;
+	
+	// Overrides to allow the addition and removal of optional inputs
+	virtual TArray<UE::Dataflow::FPin> AddPins() override;
+	virtual bool CanRemovePin() const override;
+	virtual TArray<UE::Dataflow::FPin> GetPinsToRemove() const override;
+	virtual void OnPinRemoved(const UE::Dataflow::FPin& Pin) override;
+	virtual void PostSerialize(const FArchive& Ar) override; 
+};
\ No newline at end of file
diff --git a/Engine/Plugins/Experimental/MutableDataflow/Source/MutableDataflowEditor/Public/Nodes/MakeMutableFloatParametersArrayNode.h b/Engine/Plugins/Experimental/MutableDataflow/Source/MutableDataflowEditor/Public/Nodes/MakeMutableFloatParametersArrayNode.h
new file mode 100644
--- /dev/null
+++ b/Engine/Plugins/Experimental/MutableDataflow/Source/MutableDataflowEditor/Public/Nodes/MakeMutableFloatParametersArrayNode.h
@@ -0,0 +1,39 @@
+﻿// Copyright Epic Games, Inc. All Rights Reserved.
+#pragma once
+
+#include "Misc/Guid.h"
+#include "CoreMinimal.h"
+#include "MakeMutableParametersArrayBaseNode.h"
+#include "Dataflow/DataflowNode.h"
+#include "Dataflow/DataflowInputOutput.h"
+#include "MutableDataflowParameters.h"
+
+#include "MakeMutableFloatParametersArrayNode.generated.h"
+
+
+USTRUCT(meta = (Experimental))
+struct FMakeMutableFloatParametersArrayNode : public FMakeMutableParametersArrayBaseNode
+{
+	GENERATED_USTRUCT_BODY()
+	DATAFLOW_NODE_DEFINE_INTERNAL(FMakeMutableFloatParametersArrayNode, "MakeMutableFloatParametersArray", "Mutable|Parameters", "")
+
+private:
+	UPROPERTY()
+	TArray<FMutableFloatParameter> InputFloatParameters;
+	
+	UPROPERTY(meta = (DataflowOutput, DisplayName = "Float Parameters"))
+	TArray<FMutableFloatParameter> GroupedParameters;
+
+public:
+	
+	FMakeMutableFloatParametersArrayNode(const UE::Dataflow::FNodeParameters& InParam, FGuid InGuid = FGuid::NewGuid());
+	
+	virtual void Evaluate(UE::Dataflow::FContext& Context, const FDataflowOutput* Out) const override;
+	
+	// Overrides to allow the addition and removal of optional inputs
+	virtual TArray<UE::Dataflow::FPin> AddPins() override;
+	virtual bool CanRemovePin() const override;
+	virtual TArray<UE::Dataflow::FPin> GetPinsToRemove() const override;
+	virtual void OnPinRemoved(const UE::Dataflow::FPin& Pin) override;
+	virtual void PostSerialize(const FArchive& Ar) override; 
+};
\ No newline at end of file
diff --git a/Engine/Plugins/Experimental/MutableDataflow/Source/MutableDataflowEditor/Public/Nodes/MakeMutableMaterialParametersArrayNode.h b/Engine/Plugins/Experimental/MutableDataflow/Source/MutableDataflowEditor/Public/Nodes/MakeMutableMaterialParametersArrayNode.h
--- a/Engine/Plugins/Experimental/MutableDataflow/Source/MutableDataflowEditor/Public/Nodes/MakeMutableMaterialParametersArrayNode.h
+++ b/Engine/Plugins/Experimental/MutableDataflow/Source/MutableDataflowEditor/Public/Nodes/MakeMutableMaterialParametersArrayNode.h
@@ -15,7 +15,7 @@ USTRUCT(meta = (Experimental))
 struct FMakeMutableMaterialParametersArrayNode : public FMakeMutableParametersArrayBaseNode
 {
 	GENERATED_USTRUCT_BODY()
-	DATAFLOW_NODE_DEFINE_INTERNAL(FMakeMutableMaterialParametersArrayNode, "MakeMutableMaterialParametersArray", "Mutable", "")
+	DATAFLOW_NODE_DEFINE_INTERNAL(FMakeMutableMaterialParametersArrayNode, "MakeMutableMaterialParametersArray", "Mutable|Parameters", "")
 
 private:
 	UPROPERTY()
diff --git a/Engine/Plugins/Experimental/MutableDataflow/Source/MutableDataflowEditor/Public/Nodes/MakeMutableParametersArrayBaseNode.h b/Engine/Plugins/Experimental/MutableDataflow/Source/MutableDataflowEditor/Public/Nodes/MakeMutableParametersArrayBaseNode.h
--- a/Engine/Plugins/Experimental/MutableDataflow/Source/MutableDataflowEditor/Public/Nodes/MakeMutableParametersArrayBaseNode.h
+++ b/Engine/Plugins/Experimental/MutableDataflow/Source/MutableDataflowEditor/Public/Nodes/MakeMutableParametersArrayBaseNode.h
@@ -13,7 +13,7 @@ USTRUCT(meta = (Experimental))
 struct FMakeMutableParametersArrayBaseNode : public FDataflowNode
 {
 	GENERATED_USTRUCT_BODY()
-	DATAFLOW_NODE_DEFINE_INTERNAL(FMakeMutableParametersArrayBaseNode, "MakeParametersArray", "Mutable", "")
+	DATAFLOW_NODE_DEFINE_INTERNAL(FMakeMutableParametersArrayBaseNode, "MakeParametersArray", "Mutable|Parameters", "")
 	
 public:
 	FMakeMutableParametersArrayBaseNode(const UE::Dataflow::FNodeParameters& InParam, FGuid InGuid = FGuid::NewGuid());
@@ -79,15 +79,18 @@ template <typename T>
 void FMakeMutableParametersArrayBaseNode::EvaluateParameterNode(const TArray<T>& InputParameters, const TArray<T>& OutputParameters,  UE::Dataflow::FContext& Context,
 	const FDataflowOutput* Out) const
 {
-	// Grab the values from the inputs
-	TArray<T> Output;
-	for (const T& Parameter : InputParameters)
+	if (Out->IsA(&OutputParameters))
 	{
-		Output.Add(GetValue(Context, &Parameter));
-	}
+		// Grab the values from the inputs
+		TArray<T> Output;
+		for (const T& Parameter : InputParameters)
+		{
+			Output.Add(GetValue(Context, &Parameter));
+		}
 
-	// Provide the composed array as an output
-	SetValue(Context, Output, &OutputParameters);
+		// Provide the composed array as an output
+		SetValue(Context, Output, &OutputParameters);
+	}
 }
 
 
diff --git a/Engine/Plugins/Experimental/MutableDataflow/Source/MutableDataflowEditor/Public/Nodes/MakeMutableProjectorParametersArrayNode.h b/Engine/Plugins/Experimental/MutableDataflow/Source/MutableDataflowEditor/Public/Nodes/MakeMutableProjectorParametersArrayNode.h
new file mode 100644
--- /dev/null
+++ b/Engine/Plugins/Experimental/MutableDataflow/Source/MutableDataflowEditor/Public/Nodes/MakeMutableProjectorParametersArrayNode.h
@@ -0,0 +1,39 @@
+﻿// Copyright Epic Games, Inc. All Rights Reserved.
+#pragma once
+
+#include "Misc/Guid.h"
+#include "CoreMinimal.h"
+#include "MakeMutableParametersArrayBaseNode.h"
+#include "Dataflow/DataflowNode.h"
+#include "Dataflow/DataflowInputOutput.h"
+#include "MutableDataflowParameters.h"
+
+#include "MakeMutableProjectorParametersArrayNode.generated.h"
+
+
+USTRUCT(meta = (Experimental))
+struct FMakeMutableProjectorParametersArrayNode : public FMakeMutableParametersArrayBaseNode
+{
+	GENERATED_USTRUCT_BODY()
+	DATAFLOW_NODE_DEFINE_INTERNAL(FMakeMutableProjectorParametersArrayNode, "MakeMutableProjectorParametersArray", "Mutable|Parameters", "")
+
+private:
+	UPROPERTY()
+	TArray<FMutableProjectorParameter> InputProjectorParameters;
+	
+	UPROPERTY(meta = (DataflowOutput, DisplayName = "Projector Parameters"))
+	TArray<FMutableProjectorParameter> GroupedParameters;
+
+public:
+	
+	FMakeMutableProjectorParametersArrayNode(const UE::Dataflow::FNodeParameters& InParam, FGuid InGuid = FGuid::NewGuid());
+	
+	virtual void Evaluate(UE::Dataflow::FContext& Context, const FDataflowOutput* Out) const override;
+	
+	// Overrides to allow the addition and removal of optional inputs
+	virtual TArray<UE::Dataflow::FPin> AddPins() override;
+	virtual bool CanRemovePin() const override;
+	virtual TArray<UE::Dataflow::FPin> GetPinsToRemove() const override;
+	virtual void OnPinRemoved(const UE::Dataflow::FPin& Pin) override;
+	virtual void PostSerialize(const FArchive& Ar) override; 
+};
\ No newline at end of file
diff --git a/Engine/Plugins/Experimental/MutableDataflow/Source/MutableDataflowEditor/Public/Nodes/MakeMutableSkeletalMeshParametersArrayNode.h b/Engine/Plugins/Experimental/MutableDataflow/Source/MutableDataflowEditor/Public/Nodes/MakeMutableSkeletalMeshParametersArrayNode.h
--- a/Engine/Plugins/Experimental/MutableDataflow/Source/MutableDataflowEditor/Public/Nodes/MakeMutableSkeletalMeshParametersArrayNode.h
+++ b/Engine/Plugins/Experimental/MutableDataflow/Source/MutableDataflowEditor/Public/Nodes/MakeMutableSkeletalMeshParametersArrayNode.h
@@ -15,7 +15,7 @@ USTRUCT(meta = (Experimental))
 struct FMakeMutableSkeletalMeshParametersArrayNode : public FMakeMutableParametersArrayBaseNode
 {
 	GENERATED_USTRUCT_BODY()
-	DATAFLOW_NODE_DEFINE_INTERNAL(FMakeMutableSkeletalMeshParametersArrayNode, "MakeMutableSkeletalMeshParametersArray", "Mutable", "")
+	DATAFLOW_NODE_DEFINE_INTERNAL(FMakeMutableSkeletalMeshParametersArrayNode, "MakeMutableSkeletalMeshParametersArray", "Mutable|Parameters", "")
 
 private:
 	UPROPERTY()
diff --git a/Engine/Plugins/Experimental/MutableDataflow/Source/MutableDataflowEditor/Public/Nodes/MakeMutableTextureParametersArrayNode.h b/Engine/Plugins/Experimental/MutableDataflow/Source/MutableDataflowEditor/Public/Nodes/MakeMutableTextureParametersArrayNode.h
--- a/Engine/Plugins/Experimental/MutableDataflow/Source/MutableDataflowEditor/Public/Nodes/MakeMutableTextureParametersArrayNode.h
+++ b/Engine/Plugins/Experimental/MutableDataflow/Source/MutableDataflowEditor/Public/Nodes/MakeMutableTextureParametersArrayNode.h
@@ -15,7 +15,7 @@ USTRUCT(meta = (Experimental))
 struct FMakeMutableTextureParametersArrayNode : public FMakeMutableParametersArrayBaseNode
 {
 	GENERATED_USTRUCT_BODY()
-	DATAFLOW_NODE_DEFINE_INTERNAL(FMakeMutableTextureParametersArrayNode, "MakeMutableTextureParametersArray", "Mutable", "")
+	DATAFLOW_NODE_DEFINE_INTERNAL(FMakeMutableTextureParametersArrayNode, "MakeMutableTextureParametersArray", "Mutable|Parameters", "")
 
 private:
 	UPROPERTY()
diff --git a/Engine/Plugins/Experimental/MutableDataflow/Source/MutableDataflowEditor/Public/Nodes/MakeMutableTransformParametersArrayNode.h b/Engine/Plugins/Experimental/MutableDataflow/Source/MutableDataflowEditor/Public/Nodes/MakeMutableTransformParametersArrayNode.h
new file mode 100644
--- /dev/null
+++ b/Engine/Plugins/Experimental/MutableDataflow/Source/MutableDataflowEditor/Public/Nodes/MakeMutableTransformParametersArrayNode.h
@@ -0,0 +1,39 @@
+﻿// Copyright Epic Games, Inc. All Rights Reserved.
+#pragma once
+
+#include "Misc/Guid.h"
+#include "CoreMinimal.h"
+#include "MakeMutableParametersArrayBaseNode.h"
+#include "Dataflow/DataflowNode.h"
+#include "Dataflow/DataflowInputOutput.h"
+#include "MutableDataflowParameters.h"
+
+#include "MakeMutableTransformParametersArrayNode.generated.h"
+
+
+USTRUCT(meta = (Experimental))
+struct FMakeMutableTransformParametersArrayNode : public FMakeMutableParametersArrayBaseNode
+{
+	GENERATED_USTRUCT_BODY()
+	DATAFLOW_NODE_DEFINE_INTERNAL(FMakeMutableTransformParametersArrayNode, "MakeMutableTransformParametersArray", "Mutable|Parameters", "")
+
+private:
+	UPROPERTY()
+	TArray<FMutableTransformParameter> InputTransformParameters;
+	
+	UPROPERTY(meta = (DataflowOutput, DisplayName = "Transform Parameters"))
+	TArray<FMutableTransformParameter> GroupedParameters;
+
+public:
+	
+	FMakeMutableTransformParametersArrayNode(const UE::Dataflow::FNodeParameters& InParam, FGuid InGuid = FGuid::NewGuid());
+	
+	virtual void Evaluate(UE::Dataflow::FContext& Context, const FDataflowOutput* Out) const override;
+	
+	// Overrides to allow the addition and removal of optional inputs
+	virtual TArray<UE::Dataflow::FPin> AddPins() override;
+	virtual bool CanRemovePin() const override;
+	virtual TArray<UE::Dataflow::FPin> GetPinsToRemove() const override;
+	virtual void OnPinRemoved(const UE::Dataflow::FPin& Pin) override;
+	virtual void PostSerialize(const FArchive& Ar) override; 
+};
\ No newline at end of file
diff --git a/Engine/Plugins/Experimental/MutableDataflow/Source/MutableDataflowEditor/Public/Nodes/MakeMutableVectorParametersArrayNode.h b/Engine/Plugins/Experimental/MutableDataflow/Source/MutableDataflowEditor/Public/Nodes/MakeMutableVectorParametersArrayNode.h
new file mode 100644
--- /dev/null
+++ b/Engine/Plugins/Experimental/MutableDataflow/Source/MutableDataflowEditor/Public/Nodes/MakeMutableVectorParametersArrayNode.h
@@ -0,0 +1,39 @@
+﻿// Copyright Epic Games, Inc. All Rights Reserved.
+#pragma once
+
+#include "Misc/Guid.h"
+#include "CoreMinimal.h"
+#include "MakeMutableParametersArrayBaseNode.h"
+#include "Dataflow/DataflowNode.h"
+#include "Dataflow/DataflowInputOutput.h"
+#include "MutableDataflowParameters.h"
+
+#include "MakeMutableVectorParametersArrayNode.generated.h"
+
+
+USTRUCT(meta = (Experimental))
+struct FMakeMutableVectorParametersArrayNode : public FMakeMutableParametersArrayBaseNode
+{
+	GENERATED_USTRUCT_BODY()
+	DATAFLOW_NODE_DEFINE_INTERNAL(FMakeMutableVectorParametersArrayNode, "MakeMutableVectorParametersArray", "Mutable|Parameters", "")
+
+private:
+	UPROPERTY()
+	TArray<FMutableVectorParameter> InputVectorParameters;
+	
+	UPROPERTY(meta = (DataflowOutput, DisplayName = "Vector Parameters"))
+	TArray<FMutableVectorParameter> GroupedParameters;
+
+public:
+	
+	FMakeMutableVectorParametersArrayNode(const UE::Dataflow::FNodeParameters& InParam, FGuid InGuid = FGuid::NewGuid());
+	
+	virtual void Evaluate(UE::Dataflow::FContext& Context, const FDataflowOutput* Out) const override;
+	
+	// Overrides to allow the addition and removal of optional inputs
+	virtual TArray<UE::Dataflow::FPin> AddPins() override;
+	virtual bool CanRemovePin() const override;
+	virtual TArray<UE::Dataflow::FPin> GetPinsToRemove() const override;
+	virtual void OnPinRemoved(const UE::Dataflow::FPin& Pin) override;
+	virtual void PostSerialize(const FArchive& Ar) override; 
+};
\ No newline at end of file
diff --git a/Engine/Plugins/Experimental/MutableDataflow/Source/MutableDataflowEditor/Public/Nodes/MutableBoolParameterNode.h b/Engine/Plugins/Experimental/MutableDataflow/Source/MutableDataflowEditor/Public/Nodes/MutableBoolParameterNode.h
new file mode 100644
--- /dev/null
+++ b/Engine/Plugins/Experimental/MutableDataflow/Source/MutableDataflowEditor/Public/Nodes/MutableBoolParameterNode.h
@@ -0,0 +1,30 @@
+﻿// Copyright Epic Games, Inc. All Rights Reserved.
+#pragma once
+
+#include "Dataflow/DataflowNode.h"
+#include "MutableDataflowParameters.h"
+#include "MutableParameterNode.h"
+#include "Misc/Guid.h"
+#include "CoreMinimal.h"
+
+#include "MutableBoolParameterNode.generated.h"
+
+
+USTRUCT(meta = (Experimental))
+struct FMutableBoolParameterNode : public FMutableParameterNode
+{
+	GENERATED_USTRUCT_BODY()
+	DATAFLOW_NODE_DEFINE_INTERNAL(FMutableBoolParameterNode, "MutableBoolParameterNode", "Mutable|Parameters", "")
+
+	UPROPERTY(EditAnywhere, Category = "Mutable", meta = (DataflowInput))
+	bool Bool;
+	
+	UPROPERTY(meta = (DataflowOutput))
+	FMutableBoolParameter BoolParameter;
+
+public:
+	
+	FMutableBoolParameterNode(const UE::Dataflow::FNodeParameters& InParam, FGuid InGuid = FGuid::NewGuid());
+	
+	virtual void Evaluate(UE::Dataflow::FContext& Context, const FDataflowOutput* Out) const override;
+};
\ No newline at end of file
diff --git a/Engine/Plugins/Experimental/MutableDataflow/Source/MutableDataflowEditor/Public/Nodes/MutableEnumParameterNode.h b/Engine/Plugins/Experimental/MutableDataflow/Source/MutableDataflowEditor/Public/Nodes/MutableEnumParameterNode.h
new file mode 100644
--- /dev/null
+++ b/Engine/Plugins/Experimental/MutableDataflow/Source/MutableDataflowEditor/Public/Nodes/MutableEnumParameterNode.h
@@ -0,0 +1,30 @@
+﻿// Copyright Epic Games, Inc. All Rights Reserved.
+#pragma once
+
+#include "Dataflow/DataflowNode.h"
+#include "MutableDataflowParameters.h"
+#include "MutableParameterNode.h"
+#include "Misc/Guid.h"
+#include "CoreMinimal.h"
+
+#include "MutableEnumParameterNode.generated.h"
+
+
+USTRUCT(meta = (Experimental))
+struct FMutableEnumParameterNode : public FMutableParameterNode
+{
+	GENERATED_USTRUCT_BODY()
+	DATAFLOW_NODE_DEFINE_INTERNAL(FMutableEnumParameterNode, "MutableEnumParameterNode", "Mutable|Parameters", "")
+
+	UPROPERTY(EditAnywhere, Category = "Mutable", meta = (DataflowInput))
+	FString EnumOption;
+	
+	UPROPERTY(meta = (DataflowOutput))
+	FMutableEnumParameter EnumParameter;
+
+public:
+	
+	FMutableEnumParameterNode(const UE::Dataflow::FNodeParameters& InParam, FGuid InGuid = FGuid::NewGuid());
+	
+	virtual void Evaluate(UE::Dataflow::FContext& Context, const FDataflowOutput* Out) const override;
+};
\ No newline at end of file
diff --git a/Engine/Plugins/Experimental/MutableDataflow/Source/MutableDataflowEditor/Public/Nodes/MutableFloatParameterNode.h b/Engine/Plugins/Experimental/MutableDataflow/Source/MutableDataflowEditor/Public/Nodes/MutableFloatParameterNode.h
new file mode 100644
--- /dev/null
+++ b/Engine/Plugins/Experimental/MutableDataflow/Source/MutableDataflowEditor/Public/Nodes/MutableFloatParameterNode.h
@@ -0,0 +1,30 @@
+﻿// Copyright Epic Games, Inc. All Rights Reserved.
+#pragma once
+
+#include "Dataflow/DataflowNode.h"
+#include "MutableDataflowParameters.h"
+#include "MutableParameterNode.h"
+#include "Misc/Guid.h"
+#include "CoreMinimal.h"
+
+#include "MutableFloatParameterNode.generated.h"
+
+
+USTRUCT(meta = (Experimental))
+struct FMutableFloatParameterNode : public FMutableParameterNode
+{
+	GENERATED_USTRUCT_BODY()
+	DATAFLOW_NODE_DEFINE_INTERNAL(FMutableFloatParameterNode, "MutableFloatParameterNode", "Mutable|Parameters", "")
+
+	UPROPERTY(EditAnywhere, Category = "Mutable", meta = (DataflowInput))
+	float Float;
+	
+	UPROPERTY(meta = (DataflowOutput))
+	FMutableFloatParameter FloatParameter;
+
+public:
+	
+	FMutableFloatParameterNode(const UE::Dataflow::FNodeParameters& InParam, FGuid InGuid = FGuid::NewGuid());
+	
+	virtual void Evaluate(UE::Dataflow::FContext& Context, const FDataflowOutput* Out) const override;
+};
\ No newline at end of file
diff --git a/Engine/Plugins/Experimental/MutableDataflow/Source/MutableDataflowEditor/Public/Nodes/MutableMaterialParameterNode.h b/Engine/Plugins/Experimental/MutableDataflow/Source/MutableDataflowEditor/Public/Nodes/MutableMaterialParameterNode.h
--- a/Engine/Plugins/Experimental/MutableDataflow/Source/MutableDataflowEditor/Public/Nodes/MutableMaterialParameterNode.h
+++ b/Engine/Plugins/Experimental/MutableDataflow/Source/MutableDataflowEditor/Public/Nodes/MutableMaterialParameterNode.h
@@ -5,17 +5,15 @@
 #include "Misc/Guid.h"
 #include "CoreMinimal.h"
 #include "MutableDataflowParameters.h"
+#include "MutableParameterNode.h"
 
 #include "MutableMaterialParameterNode.generated.h"
 
 USTRUCT(meta = (Experimental))
-struct FMutableMaterialParameterNode : public FDataflowNode
+struct FMutableMaterialParameterNode : public FMutableParameterNode
 {
 	GENERATED_USTRUCT_BODY()
-	DATAFLOW_NODE_DEFINE_INTERNAL(FMutableMaterialParameterNode, "MutableMaterialParameter", "Mutable", "")
-
-	UPROPERTY(EditAnywhere, Category = "Mutable", meta = (DataflowInput))
-	FString ParameterName;
+	DATAFLOW_NODE_DEFINE_INTERNAL(FMutableMaterialParameterNode, "MutableMaterialParameter", "Mutable|Parameters", "")
 
 	UPROPERTY(EditAnywhere, Category = "Mutable", meta = (DataflowInput))
 	TObjectPtr<UMaterialInterface> Material;
diff --git a/Engine/Plugins/Experimental/MutableDataflow/Source/MutableDataflowEditor/Public/Nodes/MutableParameterNode.h b/Engine/Plugins/Experimental/MutableDataflow/Source/MutableDataflowEditor/Public/Nodes/MutableParameterNode.h
new file mode 100644
--- /dev/null
+++ b/Engine/Plugins/Experimental/MutableDataflow/Source/MutableDataflowEditor/Public/Nodes/MutableParameterNode.h
@@ -0,0 +1,26 @@
+﻿// Copyright Epic Games, Inc. All Rights Reserved.
+#pragma once
+
+#include "Dataflow/DataflowNode.h"
+#include "Misc/Guid.h"
+#include "CoreMinimal.h"
+
+#include "MutableParameterNode.generated.h"
+
+
+USTRUCT(meta = (Experimental))
+struct FMutableParameterNode : public FDataflowNode
+{
+	GENERATED_USTRUCT_BODY()
+	DATAFLOW_NODE_DEFINE_INTERNAL(FMutableParameterNode, "MutableParameterBaseNode", "Mutable|Parameters", "")
+
+	UPROPERTY(EditAnywhere, Category = "Mutable", meta = (DataflowInput))
+	FString ParameterName;
+	
+public:
+	
+	FMutableParameterNode(const UE::Dataflow::FNodeParameters& InParam, FGuid InGuid = FGuid::NewGuid());
+
+protected:
+	FString GetParameterName(UE::Dataflow::FContext& Context) const;
+};
\ No newline at end of file
diff --git a/Engine/Plugins/Experimental/MutableDataflow/Source/MutableDataflowEditor/Public/Nodes/MutableProjectorParameterNode.h b/Engine/Plugins/Experimental/MutableDataflow/Source/MutableDataflowEditor/Public/Nodes/MutableProjectorParameterNode.h
new file mode 100644
--- /dev/null
+++ b/Engine/Plugins/Experimental/MutableDataflow/Source/MutableDataflowEditor/Public/Nodes/MutableProjectorParameterNode.h
@@ -0,0 +1,29 @@
+﻿// Copyright Epic Games, Inc. All Rights Reserved.
+#pragma once
+
+#include "Dataflow/DataflowNode.h"
+#include "MutableDataflowParameters.h"
+#include "MutableParameterNode.h"
+#include "Misc/Guid.h"
+#include "CoreMinimal.h"
+
+#include "MutableProjectorParameterNode.generated.h"
+
+USTRUCT(meta = (Experimental))
+struct FMutableProjectorParameterNode : public FMutableParameterNode
+{
+	GENERATED_USTRUCT_BODY()
+	DATAFLOW_NODE_DEFINE_INTERNAL(FMutableProjectorParameterNode, "MutableProjectorParameterNode", "Mutable|Parameters", "")
+
+	UPROPERTY(EditAnywhere, Category = "Mutable", meta = (DataflowInput))
+	FCustomizableObjectProjector Projector;
+	
+	UPROPERTY(meta = (DataflowOutput))
+	FMutableProjectorParameter ProjectorParameter;
+
+public:
+	
+	FMutableProjectorParameterNode(const UE::Dataflow::FNodeParameters& InParam, FGuid InGuid = FGuid::NewGuid());
+	
+	virtual void Evaluate(UE::Dataflow::FContext& Context, const FDataflowOutput* Out) const override;
+};
\ No newline at end of file
diff --git a/Engine/Plugins/Experimental/MutableDataflow/Source/MutableDataflowEditor/Public/Nodes/MutableSkeletalMeshParameterNode.h b/Engine/Plugins/Experimental/MutableDataflow/Source/MutableDataflowEditor/Public/Nodes/MutableSkeletalMeshParameterNode.h
--- a/Engine/Plugins/Experimental/MutableDataflow/Source/MutableDataflowEditor/Public/Nodes/MutableSkeletalMeshParameterNode.h
+++ b/Engine/Plugins/Experimental/MutableDataflow/Source/MutableDataflowEditor/Public/Nodes/MutableSkeletalMeshParameterNode.h
@@ -3,6 +3,7 @@
 
 #include "Dataflow/DataflowNode.h"
 #include "MutableDataflowParameters.h"
+#include "MutableParameterNode.h"
 #include "Misc/Guid.h"
 #include "CoreMinimal.h"
 
@@ -10,13 +11,10 @@
 
 
 USTRUCT(meta = (Experimental))
-struct FMutableSkeletalMeshParameterNode : public FDataflowNode
+struct FMutableSkeletalMeshParameterNode : public FMutableParameterNode
 {
 	GENERATED_USTRUCT_BODY()
-	DATAFLOW_NODE_DEFINE_INTERNAL(FMutableSkeletalMeshParameterNode, "MutableSkeletalMeshParameter", "Mutable", "")
-	
-	UPROPERTY(EditAnywhere, Category = "Mutable", meta = (DataflowInput))
-	FString ParameterName;
+	DATAFLOW_NODE_DEFINE_INTERNAL(FMutableSkeletalMeshParameterNode, "MutableSkeletalMeshParameter", "Mutable|Parameters", "")
 
 	UPROPERTY(EditAnywhere, Category = "Mutable", meta = (DataflowInput))
 	TObjectPtr<USkeletalMesh> SkeletalMesh;
@@ -29,4 +27,4 @@ public:
 	FMutableSkeletalMeshParameterNode(const UE::Dataflow::FNodeParameters& InParam, FGuid InGuid = FGuid::NewGuid());
 	
 	virtual void Evaluate(UE::Dataflow::FContext& Context, const FDataflowOutput* Out) const override;
-};
\ No newline at end of file
+};
diff --git a/Engine/Plugins/Experimental/MutableDataflow/Source/MutableDataflowEditor/Public/Nodes/MutableTextureParameterNode.h b/Engine/Plugins/Experimental/MutableDataflow/Source/MutableDataflowEditor/Public/Nodes/MutableTextureParameterNode.h
--- a/Engine/Plugins/Experimental/MutableDataflow/Source/MutableDataflowEditor/Public/Nodes/MutableTextureParameterNode.h
+++ b/Engine/Plugins/Experimental/MutableDataflow/Source/MutableDataflowEditor/Public/Nodes/MutableTextureParameterNode.h
@@ -5,17 +5,15 @@
 #include "Misc/Guid.h"
 #include "CoreMinimal.h"
 #include "MutableDataflowParameters.h"
+#include "MutableParameterNode.h"
 
 #include "MutableTextureParameterNode.generated.h"
 
 USTRUCT(meta = (Experimental))
-struct FMutableTextureParameterNode : public FDataflowNode
+struct FMutableTextureParameterNode : public FMutableParameterNode
 {
 	GENERATED_USTRUCT_BODY()
-	DATAFLOW_NODE_DEFINE_INTERNAL(FMutableTextureParameterNode, "MutableTextureParameter", "Mutable", "")
-
-	UPROPERTY(EditAnywhere, Category = "Mutable", meta = (DataflowInput))
-	FString ParameterName;
+	DATAFLOW_NODE_DEFINE_INTERNAL(FMutableTextureParameterNode, "MutableTextureParameter", "Mutable|Parameters", "")
 
 	UPROPERTY(EditAnywhere, Category = "Mutable", meta = (DataflowInput))
 	TObjectPtr<UTexture2D> Texture;
diff --git a/Engine/Plugins/Experimental/MutableDataflow/Source/MutableDataflowEditor/Public/Nodes/MutableTransformParameterNode.h b/Engine/Plugins/Experimental/MutableDataflow/Source/MutableDataflowEditor/Public/Nodes/MutableTransformParameterNode.h
new file mode 100644
--- /dev/null
+++ b/Engine/Plugins/Experimental/MutableDataflow/Source/MutableDataflowEditor/Public/Nodes/MutableTransformParameterNode.h
@@ -0,0 +1,30 @@
+﻿// Copyright Epic Games, Inc. All Rights Reserved.
+#pragma once
+
+#include "Dataflow/DataflowNode.h"
+#include "MutableDataflowParameters.h"
+#include "MutableParameterNode.h"
+#include "Misc/Guid.h"
+#include "CoreMinimal.h"
+
+#include "MutableTransformParameterNode.generated.h"
+
+
+USTRUCT(meta = (Experimental))
+struct FMutableTransformParameterNode : public FMutableParameterNode
+{
+	GENERATED_USTRUCT_BODY()
+	DATAFLOW_NODE_DEFINE_INTERNAL(FMutableTransformParameterNode, "MutableTransformParameterNode", "Mutable|Parameters", "")
+
+	UPROPERTY(EditAnywhere, Category = "Mutable", meta = (DataflowInput))
+	FTransform Transform;
+	
+	UPROPERTY(meta = (DataflowOutput))
+	FMutableTransformParameter TransformParameter;
+
+public:
+	
+	FMutableTransformParameterNode(const UE::Dataflow::FNodeParameters& InParam, FGuid InGuid = FGuid::NewGuid());
+	
+	virtual void Evaluate(UE::Dataflow::FContext& Context, const FDataflowOutput* Out) const override;
+};
diff --git a/Engine/Plugins/Experimental/MutableDataflow/Source/MutableDataflowEditor/Public/Nodes/MutableVectorParameterNode.h b/Engine/Plugins/Experimental/MutableDataflow/Source/MutableDataflowEditor/Public/Nodes/MutableVectorParameterNode.h
new file mode 100644
--- /dev/null
+++ b/Engine/Plugins/Experimental/MutableDataflow/Source/MutableDataflowEditor/Public/Nodes/MutableVectorParameterNode.h
@@ -0,0 +1,29 @@
+﻿// Copyright Epic Games, Inc. All Rights Reserved.
+#pragma once
+
+#include "Dataflow/DataflowNode.h"
+#include "MutableDataflowParameters.h"
+#include "MutableParameterNode.h"
+#include "Misc/Guid.h"
+#include "CoreMinimal.h"
+
+#include "MutableVectorParameterNode.generated.h"
+
+USTRUCT(meta = (Experimental))
+struct FMutableVectorParameterNode : public FMutableParameterNode
+{
+	GENERATED_USTRUCT_BODY()
+	DATAFLOW_NODE_DEFINE_INTERNAL(FMutableVectorParameterNode, "MutableVectorParameterNode", "Mutable|Parameters", "")
+
+	UPROPERTY(EditAnywhere, Category = "Mutable", meta = (DataflowInput))
+	FLinearColor Color;
+	
+	UPROPERTY(meta = (DataflowOutput))
+	FMutableVectorParameter VectorParameter;
+
+public:
+	
+	FMutableVectorParameterNode(const UE::Dataflow::FNodeParameters& InParam, FGuid InGuid = FGuid::NewGuid());
+	
+	virtual void Evaluate(UE::Dataflow::FContext& Context, const FDataflowOutput* Out) const override;
+};
