[mutable] Dataflow Instanced Struct Parameters.
- Added support for UStructs to be used as parameters of the COI (in the context of Dataflow).
- Implemented usage of a new type of policy that takes care of handling connections with UStruct objects.
- Added support nodes for the new type of dataflow mutable parameter node

#rnx

diff --git a/Engine/Plugins/Experimental/MutableDataflow/Source/MutableDataflowEditor/Private/MutableDataflowEditorModule.cpp b/Engine/Plugins/Experimental/MutableDataflow/Source/MutableDataflowEditor/Private/MutableDataflowEditorModule.cpp
--- a/Engine/Plugins/Experimental/MutableDataflow/Source/MutableDataflowEditor/Private/MutableDataflowEditorModule.cpp
+++ b/Engine/Plugins/Experimental/MutableDataflow/Source/MutableDataflowEditor/Private/MutableDataflowEditorModule.cpp
@@ -10,6 +10,7 @@
 #include "Nodes/MakeMutableBoolParametersArrayNode.h"
 #include "Nodes/MakeMutableEnumParametersArrayNode.h"
 #include "Nodes/MakeMutableFloatParametersArrayNode.h"
+#include "Nodes/MakeMutableInstancedStructParametersArrayNode.h"
 #include "Nodes/MutableMaterialParameterNode.h"
 #include "Nodes/MutableSkeletalMeshParameterNode.h"
 #include "Nodes/MutableTextureParameterNode.h"
@@ -25,6 +26,7 @@
 #include "Nodes/MutableProjectorParameterNode.h"
 #include "Nodes/MutableTransformParameterNode.h"
 #include "Nodes/MutableVectorParameterNode.h"
+#include "Nodes/MutableInstancedStructParameterNode.h"
 
 IMPLEMENT_MODULE(FMutableDataflowEditorModule, MutableDataflowEditor);
 
@@ -50,7 +52,8 @@ void FMutableDataflowEditorModule::StartupModule()
 	DATAFLOW_NODE_REGISTER_CREATION_FACTORY(FMutableTransformParameterNode);
 	DATAFLOW_NODE_REGISTER_CREATION_FACTORY(FMutableVectorParameterNode);
 	DATAFLOW_NODE_REGISTER_CREATION_FACTORY(FMutableEnumParameterNode);
-	
+	DATAFLOW_NODE_REGISTER_CREATION_FACTORY(FMutableInstancedStructParameterNode);
+
 	// Parameter array generation nodes
 	DATAFLOW_NODE_REGISTER_CREATION_FACTORY(FMakeMutableSkeletalMeshParametersArrayNode);
 	DATAFLOW_NODE_REGISTER_CREATION_FACTORY(FMakeMutableMaterialParametersArrayNode);
@@ -61,7 +64,8 @@ void FMutableDataflowEditorModule::StartupModule()
 	DATAFLOW_NODE_REGISTER_CREATION_FACTORY(FMakeMutableTransformParametersArrayNode);
 	DATAFLOW_NODE_REGISTER_CREATION_FACTORY(FMakeMutableVectorParametersArrayNode);
 	DATAFLOW_NODE_REGISTER_CREATION_FACTORY(FMakeMutableEnumParametersArrayNode);
-	
+	DATAFLOW_NODE_REGISTER_CREATION_FACTORY(FMakeMutableInstancedStructParametersArrayNode);
+
 	// Parameter node to array automatic conversion nodes
 	UE_DATAFLOW_REGISTER_AUTOCONVERT(FMutableSkeletalMeshParameter, TArray<FMutableSkeletalMeshParameter>, FMakeMutableSkeletalMeshParametersArrayNode);
 	UE_DATAFLOW_REGISTER_AUTOCONVERT(FMutableMaterialParameter, TArray<FMutableMaterialParameter>, FMakeMutableMaterialParametersArrayNode);
@@ -72,4 +76,6 @@ void FMutableDataflowEditorModule::StartupModule()
 	UE_DATAFLOW_REGISTER_AUTOCONVERT(FMutableTransformParameter, TArray<FMutableTransformParameter>, FMakeMutableTransformParametersArrayNode);
 	UE_DATAFLOW_REGISTER_AUTOCONVERT(FMutableVectorParameter, TArray<FMutableVectorParameter>, FMakeMutableVectorParametersArrayNode);
 	UE_DATAFLOW_REGISTER_AUTOCONVERT(FMutableEnumParameter, TArray<FMutableEnumParameter>, FMakeMutableEnumParametersArrayNode);
+	UE_DATAFLOW_REGISTER_AUTOCONVERT(FMutableInstancedStructParameter, TArray<FMutableInstancedStructParameter>, FMakeMutableInstancedStructParametersArrayNode);
+
 }
diff --git a/Engine/Plugins/Experimental/MutableDataflow/Source/MutableDataflowEditor/Private/Nodes/MakeMutableInstancedStructParametersArrayNode.cpp b/Engine/Plugins/Experimental/MutableDataflow/Source/MutableDataflowEditor/Private/Nodes/MakeMutableInstancedStructParametersArrayNode.cpp
new file mode 100644
--- /dev/null
+++ b/Engine/Plugins/Experimental/MutableDataflow/Source/MutableDataflowEditor/Private/Nodes/MakeMutableInstancedStructParametersArrayNode.cpp
@@ -0,0 +1,52 @@
+﻿// Copyright Epic Games, Inc. All Rights Reserved.
+
+
+#include "../../Public/Nodes/MakeMutableInstancedStructParametersArrayNode.h"
+#include "Nodes/MakeMutableInstancedStructParametersArrayNode.h"
+#include "Dataflow/DataflowNode.h"
+
+
+FMakeMutableInstancedStructParametersArrayNode::FMakeMutableInstancedStructParametersArrayNode(const UE::Dataflow::FNodeParameters& InParam, FGuid InGuid)
+	: Super(InParam, InGuid)
+{
+	RegisterOutputConnection(&GroupedParameters);
+
+	// Create as many inputs as the node requires to have by default
+	AddDefaultInputs(InputStructParameters);
+}
+
+
+void FMakeMutableInstancedStructParametersArrayNode::Evaluate(UE::Dataflow::FContext& Context, const FDataflowOutput* Out) const
+{
+	EvaluateParameterNode(InputStructParameters, GroupedParameters, Context, Out);
+}
+
+
+TArray<UE::Dataflow::FPin> FMakeMutableInstancedStructParametersArrayNode::AddPins()
+{
+	return AddParameterPin(InputStructParameters);
+}
+
+
+bool FMakeMutableInstancedStructParametersArrayNode::CanRemovePin() const
+{
+	return CanRemoveParameterPin(InputStructParameters);
+}
+
+
+TArray<UE::Dataflow::FPin> FMakeMutableInstancedStructParametersArrayNode::GetPinsToRemove() const
+{
+	return GetParameterPinsToRemove(InputStructParameters);
+}
+
+
+void FMakeMutableInstancedStructParametersArrayNode::OnPinRemoved(const UE::Dataflow::FPin& Pin)
+{
+	OnParameterPinRemoved(InputStructParameters, Pin);
+}
+
+
+void FMakeMutableInstancedStructParametersArrayNode::PostSerialize(const FArchive& Ar)
+{
+	PostNodeSerialize(InputStructParameters, Ar);
+}
diff --git a/Engine/Plugins/Experimental/MutableDataflow/Source/MutableDataflowEditor/Private/Nodes/MutableInstancedStructParameterNode.cpp b/Engine/Plugins/Experimental/MutableDataflow/Source/MutableDataflowEditor/Private/Nodes/MutableInstancedStructParameterNode.cpp
new file mode 100644
--- /dev/null
+++ b/Engine/Plugins/Experimental/MutableDataflow/Source/MutableDataflowEditor/Private/Nodes/MutableInstancedStructParameterNode.cpp
@@ -0,0 +1,29 @@
+﻿// Copyright Epic Games, Inc. All Rights Reserved.
+
+#include "Nodes/MutableInstancedStructParameterNode.h"
+
+#include "Dataflow/DataflowEngineUtil.h"
+#include "Dataflow/DataflowObject.h"
+#include "Dataflow/DataflowAnyType.h"
+
+
+FMutableInstancedStructParameterNode::FMutableInstancedStructParameterNode(const UE::Dataflow::FNodeParameters& InParam, FGuid InGuid)
+	: Super(InParam, InGuid)
+{
+	RegisterInputConnection(&InputStruct);
+	
+	RegisterOutputConnection(&InstancedStructParameter);
+}
+
+
+void FMutableInstancedStructParameterNode::Evaluate(UE::Dataflow::FContext& Context, const FDataflowOutput* Out) const
+{
+	if (Out->IsA(&InstancedStructParameter))
+	{
+		FMutableInstancedStructParameter Output;
+		Output.Name = GetParameterName(Context);
+		Output.InstancedStruct = GetValue(Context, &InputStruct).GetConstStructView();
+	
+		SetValue(Context, Output, &InstancedStructParameter);
+	}
+}
diff --git a/Engine/Plugins/Experimental/MutableDataflow/Source/MutableDataflowEditor/Public/MutableDataflowParameters.h b/Engine/Plugins/Experimental/MutableDataflow/Source/MutableDataflowEditor/Public/MutableDataflowParameters.h
--- a/Engine/Plugins/Experimental/MutableDataflow/Source/MutableDataflowEditor/Public/MutableDataflowParameters.h
+++ b/Engine/Plugins/Experimental/MutableDataflow/Source/MutableDataflowEditor/Public/MutableDataflowParameters.h
@@ -116,3 +116,12 @@ struct FMutableTransformParameter : public FMutableParameterBase
 	FTransform Transform;
 };
 
+
+USTRUCT(BlueprintType, meta = (Experimental))
+struct FMutableInstancedStructParameter : public FMutableParameterBase
+{
+	GENERATED_USTRUCT_BODY()
+	
+	UPROPERTY(EditAnywhere, Category = "Mutable")
+	FInstancedStruct InstancedStruct; 
+};
diff --git a/Engine/Plugins/Experimental/MutableDataflow/Source/MutableDataflowEditor/Public/Nodes/COInstanceGeneratorNode.h b/Engine/Plugins/Experimental/MutableDataflow/Source/MutableDataflowEditor/Public/Nodes/COInstanceGeneratorNode.h
--- a/Engine/Plugins/Experimental/MutableDataflow/Source/MutableDataflowEditor/Public/Nodes/COInstanceGeneratorNode.h
+++ b/Engine/Plugins/Experimental/MutableDataflow/Source/MutableDataflowEditor/Public/Nodes/COInstanceGeneratorNode.h
@@ -111,6 +111,10 @@ private:
 	TArray<FMutableTransformParameter> TransformParameters;
 	TArray<FMutableTransformParameter> CachedTransformParameters;
 	
+	UPROPERTY(EditAnywhere, Category = "Mutable Parameters", meta = (DataflowInput))
+	TArray<FMutableInstancedStructParameter> InstancedStructParameters;
+	TArray<FMutableInstancedStructParameter> CachedInstancedStructParameters;
+	
 	// Outputs
 
 	// First option : an array of structures that contain the pair of data Component name + Skeletal Mesh
diff --git a/Engine/Plugins/Experimental/MutableDataflow/Source/MutableDataflowEditor/Public/Nodes/MakeMutableInstancedStructParametersArrayNode.h b/Engine/Plugins/Experimental/MutableDataflow/Source/MutableDataflowEditor/Public/Nodes/MakeMutableInstancedStructParametersArrayNode.h
new file mode 100644
--- /dev/null
+++ b/Engine/Plugins/Experimental/MutableDataflow/Source/MutableDataflowEditor/Public/Nodes/MakeMutableInstancedStructParametersArrayNode.h
@@ -0,0 +1,39 @@
+﻿// Copyright Epic Games, Inc. All Rights Reserved.
+#pragma once
+
+#include "Misc/Guid.h"
+#include "CoreMinimal.h"
+#include "MakeMutableParametersArrayBaseNode.h"
+#include "Dataflow/DataflowNode.h"
+#include "Dataflow/DataflowInputOutput.h"
+#include "MutableDataflowParameters.h"
+
+#include "MakeMutableInstancedStructParametersArrayNode.generated.h"
+
+
+USTRUCT(meta = (Experimental))
+struct FMakeMutableInstancedStructParametersArrayNode : public FMakeMutableParametersArrayBaseNode
+{
+	GENERATED_USTRUCT_BODY()
+	DATAFLOW_NODE_DEFINE_INTERNAL(FMakeMutableInstancedStructParametersArrayNode, "MakeMutableInstancedStructParametersArray", "Mutable|Parameters", "")
+
+private:
+	UPROPERTY()
+	TArray<FMutableInstancedStructParameter> InputStructParameters;
+	
+	UPROPERTY(meta = (DataflowOutput, DisplayName = "Instanced Struct Parameters"))
+	TArray<FMutableInstancedStructParameter> GroupedParameters;
+
+public:
+	
+	FMakeMutableInstancedStructParametersArrayNode(const UE::Dataflow::FNodeParameters& InParam, FGuid InGuid = FGuid::NewGuid());
+	
+	virtual void Evaluate(UE::Dataflow::FContext& Context, const FDataflowOutput* Out) const override;
+
+	// Overrides to allow the addition and removal of optional inputs
+	virtual TArray<UE::Dataflow::FPin> AddPins() override;
+	virtual bool CanRemovePin() const override;
+	virtual TArray<UE::Dataflow::FPin> GetPinsToRemove() const override;
+	virtual void OnPinRemoved(const UE::Dataflow::FPin& Pin) override;
+	virtual void PostSerialize(const FArchive& Ar) override; 
+};
\ No newline at end of file
diff --git a/Engine/Plugins/Experimental/MutableDataflow/Source/MutableDataflowEditor/Public/Nodes/MutableInstancedStructParameterNode.h b/Engine/Plugins/Experimental/MutableDataflow/Source/MutableDataflowEditor/Public/Nodes/MutableInstancedStructParameterNode.h
new file mode 100644
--- /dev/null
+++ b/Engine/Plugins/Experimental/MutableDataflow/Source/MutableDataflowEditor/Public/Nodes/MutableInstancedStructParameterNode.h
@@ -0,0 +1,31 @@
+﻿// Copyright Epic Games, Inc. All Rights Reserved.
+#pragma once
+
+#include "Dataflow/DataflowNode.h"
+#include "Dataflow/DataflowAnyType.h"
+#include "Misc/Guid.h"
+#include "CoreMinimal.h"
+#include "MutableDataflowParameters.h"
+#include "MutableParameterNode.h"
+
+#include "MutableInstancedStructParameterNode.generated.h"
+
+
+USTRUCT(meta = (Experimental))
+struct FMutableInstancedStructParameterNode : public FMutableParameterNode
+{
+	GENERATED_USTRUCT_BODY()
+	DATAFLOW_NODE_DEFINE_INTERNAL(FMutableInstancedStructParameterNode, "MutableInstancedStructParameter", "Mutable|Parameters", "")
+	
+	UPROPERTY(meta=(DataflowInput))
+	FDataflowUStructTypes InputStruct;
+	
+	UPROPERTY(meta = (DataflowOutput))
+	FMutableInstancedStructParameter InstancedStructParameter;
+
+public:
+	
+	FMutableInstancedStructParameterNode(const UE::Dataflow::FNodeParameters& InParam, FGuid InGuid = FGuid::NewGuid());
+	
+	virtual void Evaluate(UE::Dataflow::FContext& Context, const FDataflowOutput* Out) const override;
+};
