[mutable] Added support for FInstanecdStructs as parameters of the COI

diff --git a/Engine/Plugins/Experimental/MutableDataflow/Source/MutableDataflowEditor/Private/Nodes/COInstanceGeneratorNode.cpp b/Engine/Plugins/Experimental/MutableDataflow/Source/MutableDataflowEditor/Private/Nodes/COInstanceGeneratorNode.cpp
--- a/Engine/Plugins/Experimental/MutableDataflow/Source/MutableDataflowEditor/Private/Nodes/COInstanceGeneratorNode.cpp
+++ b/Engine/Plugins/Experimental/MutableDataflow/Source/MutableDataflowEditor/Private/Nodes/COInstanceGeneratorNode.cpp
@@ -57,6 +57,7 @@ FCOInstanceGeneratorNode::FCOInstanceGeneratorNode(const UE::Dataflow::FNodePara
 	RegisterInputConnection(&VectorParameters);
 	RegisterInputConnection(&ProjectorParameters);
 	RegisterInputConnection(&TransformParameters);
+	RegisterInputConnection(&InstancedStructParameters);
 	
 	// Outputs
 	RegisterOutputConnection(&GeneratedResources);
@@ -133,6 +134,7 @@ void FCOInstanceGeneratorNode::ClearCachedParameters()
 	CachedVectorParameters.Reset();
 	CachedProjectorParameters.Reset();
 	CachedTransformParameters.Reset();
+	CachedInstancedStructParameters.Reset();
 }
 
 
@@ -219,6 +221,7 @@ void FCOInstanceGeneratorNode::CacheNodeInputs(UE::Dataflow::FContext& Context)
 	CachedVectorParameters = GetPropertyParameters(&VectorParameters,TEXT("Vector"), Context);
 	CachedProjectorParameters = GetPropertyParameters(&ProjectorParameters,TEXT("Projector"), Context);
 	CachedTransformParameters = GetPropertyParameters(&TransformParameters,TEXT("Transform"),Context);
+	CachedInstancedStructParameters = GetPropertyParameters(&InstancedStructParameters,TEXT("Instanced Struct"), Context);
 }
 
 
@@ -308,6 +311,15 @@ void FCOInstanceGeneratorNode::ApplyInstanceParameters(UCustomizableObjectInstan
 			Instance.SetTransformParameterSelectedOption(ParameterName, TransformParameter.Transform);				
 		}
 	}
+	
+	for (const FMutableInstancedStructParameter& InstancedStructParameter : CachedInstancedStructParameters)
+	{
+		const FString ParameterName = InstancedStructParameter.Name;
+		if (Instance.ContainsExternalTypeParameter(ParameterName))
+		{
+			Instance.SetExternalTypeParameterSelectedOption(ParameterName, InstancedStructParameter.InstancedStruct);				
+		}
+	}
 }
 
 
